Unreal Insights channel filtering plugin:
- Adds a new tab to the Editor, "Trace Data Filtering"
- Tab allows for setting individual Trace Channel states, meaning whether or not they should output events
- Also allows for specifying Presets, defining groups of Channels that should be enabled together
- Enabled plugin for QAGame project

diff --git a/Engine/Plugins/Developer/TraceDataFiltering/Config/BaseTraceDataFilters.ini b/Engine/Plugins/Developer/TraceDataFiltering/Config/BaseTraceDataFilters.ini
new file mode 100644
--- /dev/null
+++ b/Engine/Plugins/Developer/TraceDataFiltering/Config/BaseTraceDataFilters.ini
@@ -0,0 +1,2 @@
+[/Script/TraceDataFiltering.SharedFilterPresetContainer]
++EnginePresets=(Name="Animation",WhitelistedNames=("Animation", "Object", "Frame"))
diff --git a/Engine/Plugins/Developer/TraceDataFiltering/Source/TraceDataFiltering/Private/EventFilterService.cpp b/Engine/Plugins/Developer/TraceDataFiltering/Source/TraceDataFiltering/Private/EventFilterService.cpp
new file mode 100644
--- /dev/null
+++ b/Engine/Plugins/Developer/TraceDataFiltering/Source/TraceDataFiltering/Private/EventFilterService.cpp
@@ -0,0 +1,75 @@
+// Copyright Epic Games, Inc. All Rights Reserved.
+
+#include "EventFilterService.h"
+
+#include "Modules/ModuleManager.h"
+#include "Templates/SharedPointer.h"
+#include "Algo/Transform.h"
+#include "HAL/PlatformFilemanager.h"
+
+#include "TraceServices/Model/AnalysisSession.h"
+#include "TraceServices/ITraceServicesModule.h"
+#include "TraceServices/AnalysisService.h"
+#include "TraceServices/SessionService.h"
+
+#include "Trace/SessionTraceFilterService.h"
+
+#include "IFilterPreset.h"
+#include "FilterPresets.h"
+
+FEventFilterService& FEventFilterService::Get()
+{
+	static FEventFilterService Service;
+	return Service;
+}
+
+TSharedRef<ISessionTraceFilterService> FEventFilterService::GetFilterServiceByHandle(Trace::FSessionHandle InHandle)
+{
+	/** Find or add a per-session filter service */
+	if (TSharedPtr<ISessionTraceFilterService>* Service = PerHandleFilterService.Find(InHandle))
+	{
+		return Service->ToSharedRef();
+	}
+	else
+	{
+		TSharedPtr<ISessionTraceFilterService> NewService = nullptr;
+		for (TSharedPtr<const Trace::IAnalysisSession> Session : AnalysisSessions)
+		{
+			if (SessionService->GetSessionHandleByName(Session->GetName()) == InHandle)
+			{
+				NewService = MakeShareable(new FSessionTraceFilterService(InHandle, Session.ToSharedRef()));
+			}
+		}
+
+		if (NewService == nullptr)
+		{
+			TSharedPtr<const Trace::IAnalysisSession> Session = SessionService->StartAnalysis(InHandle);
+
+			TSharedPtr<FSessionTraceFilterService> TraceService = MakeShareable(new FSessionTraceFilterService(InHandle, Session.ToSharedRef()));
+			NewService = TraceService;
+		}
+
+		ensure(NewService != nullptr);
+		PerHandleFilterService.Add(InHandle, NewService);
+		return NewService->AsShared();
+	}
+}
+
+FEventFilterService::FEventFilterService()
+{
+	ITraceServicesModule& TraceServicesModule = FModuleManager::LoadModuleChecked<ITraceServicesModule>("TraceServices");
+	TSharedPtr<Trace::IAnalysisService> TraceAnalysisService = TraceServicesModule.GetAnalysisService();
+
+	SessionService = TraceServicesModule.GetSessionService();
+
+	/** Hook into analysis callbacks to track active sessions */
+	TraceAnalysisService->OnAnalysisStarted().AddLambda([this](TSharedRef<const Trace::IAnalysisSession> Session)
+	{
+		AnalysisSessions.Add(Session);
+	});
+
+	TraceAnalysisService->OnAnalysisFinished().AddLambda([this](TSharedRef<const Trace::IAnalysisSession> Session)
+	{
+		AnalysisSessions.Remove(Session);
+	});
+}
diff --git a/Engine/Plugins/Developer/TraceDataFiltering/Source/TraceDataFiltering/Private/EventFilterService.h b/Engine/Plugins/Developer/TraceDataFiltering/Source/TraceDataFiltering/Private/EventFilterService.h
new file mode 100644
--- /dev/null
+++ b/Engine/Plugins/Developer/TraceDataFiltering/Source/TraceDataFiltering/Private/EventFilterService.h
@@ -0,0 +1,36 @@
+// Copyright Epic Games, Inc. All Rights Reserved.
+
+#pragma once
+
+#include "Containers/Array.h"
+#include "Containers/Map.h"
+#include "ISessionTraceFilterService.h"
+
+namespace Trace
+{
+	class IAnalysisSession;
+	class ISessionService;
+	typedef uint64 FSessionHandle;
+}
+
+class FEventFilterService
+{
+protected:
+	FEventFilterService();
+public:
+	/** Singleton-getter */
+	static FEventFilterService& Get();
+
+	/** Retrieves the ISessionTraceFilterService for the provided Trace Session Handle */
+	TSharedRef<ISessionTraceFilterService> GetFilterServiceByHandle(Trace::FSessionHandle InHandle);
+
+protected:
+	/** Currently active set of Trace Analysis sessions */
+	TArray<TSharedPtr<const Trace::IAnalysisSession>> AnalysisSessions;
+	
+	/** TraceFilterServices for corresponding Trace Analysis sessions */
+	TMap<Trace::FSessionHandle, TSharedPtr<ISessionTraceFilterService>> PerHandleFilterService;
+
+	/** Cached instance of the Trace Session Service */
+	TSharedPtr<Trace::ISessionService> SessionService;
+};
\ No newline at end of file
diff --git a/Engine/Plugins/Developer/TraceDataFiltering/Source/TraceDataFiltering/Private/EventFilterStyle.cpp b/Engine/Plugins/Developer/TraceDataFiltering/Source/TraceDataFiltering/Private/EventFilterStyle.cpp
new file mode 100644
--- /dev/null
+++ b/Engine/Plugins/Developer/TraceDataFiltering/Source/TraceDataFiltering/Private/EventFilterStyle.cpp
@@ -0,0 +1,111 @@
+// Copyright Epic Games, Inc. All Rights Reserved.
+
+#include "EventFilterStyle.h"
+#include "Styling/SlateStyleRegistry.h"
+#include "Styling/SlateTypes.h"
+#include "Styling/CoreStyle.h"
+
+TSharedPtr< FSlateStyleSet > FEventFilterStyle::StyleSet = nullptr;
+
+#define IMAGE_BRUSH( RelativePath, ... ) FSlateImageBrush( StyleSet->RootToContentDir( RelativePath, TEXT(".png") ), __VA_ARGS__ )
+#define BOX_BRUSH( RelativePath, ... ) FSlateBoxBrush( StyleSet->RootToContentDir( RelativePath, TEXT(".png") ), __VA_ARGS__ )
+#define BORDER_BRUSH( RelativePath, ... ) FSlateBorderBrush( StyleSet->RootToContentDir( RelativePath, TEXT(".png") ), __VA_ARGS__ )
+#define DEFAULT_FONT(...) FCoreStyle::GetDefaultFontStyle(__VA_ARGS__)
+
+// Const icon sizes
+static const FVector2D Icon8x8(8.0f, 8.0f);
+static const FVector2D Icon9x19(9.0f, 19.0f);
+static const FVector2D Icon14x14(14.0f, 14.0f);
+static const FVector2D Icon16x16(16.0f, 16.0f);
+static const FVector2D Icon20x20(20.0f, 20.0f);
+static const FVector2D Icon22x22(22.0f, 22.0f);
+static const FVector2D Icon24x24(24.0f, 24.0f);
+static const FVector2D Icon28x28(28.0f, 28.0f);
+static const FVector2D Icon27x31(27.0f, 31.0f);
+static const FVector2D Icon26x26(26.0f, 26.0f);
+static const FVector2D Icon32x32(32.0f, 32.0f);
+static const FVector2D Icon40x40(40.0f, 40.0f);
+static const FVector2D Icon48x48(48.0f, 48.0f);
+static const FVector2D Icon75x82(75.0f, 82.0f);
+static const FVector2D Icon360x32(360.0f, 32.0f);
+static const FVector2D Icon171x39(171.0f, 39.0f);
+static const FVector2D Icon170x50(170.0f, 50.0f);
+static const FVector2D Icon267x140(170.0f, 50.0f);
+
+void FEventFilterStyle::Initialize()
+{
+	// Only register once
+	if( StyleSet.IsValid() )
+	{
+		return;
+	}
+
+	StyleSet = MakeShareable(new FSlateStyleSet("EventFilter") );
+	
+	StyleSet->SetContentRoot(FPaths::EngineContentDir() / TEXT("Editor/Slate"));
+	StyleSet->SetCoreContentRoot(FPaths::EngineContentDir() / TEXT("Slate"));
+
+	// Colors
+	{
+		StyleSet->Set("EventFilter.EnginePreset", FLinearColor(0.728f, 0.364f, 0.003f));
+		StyleSet->Set("EventFilter.SharedPreset", FLinearColor(0.003f, 0.364f, 0.728f));
+		StyleSet->Set("EventFilter.LocalPreset", FLinearColor(0.003f, 0.728f, 0.364f));
+	}
+	
+	// Icons
+	{
+		StyleSet->Set("EventFilter.State.Enabled", new IMAGE_BRUSH("Common/CheckBox_Checked", Icon16x16));
+		StyleSet->Set("EventFilter.State.Enabled_Hovered", new IMAGE_BRUSH("Common/CheckBox_Checked_Hovered", Icon16x16));
+
+		StyleSet->Set("EventFilter.State.Disabled", new IMAGE_BRUSH("Common/CheckBox", Icon16x16));
+		StyleSet->Set("EventFilter.State.Disabled_Hovered", new IMAGE_BRUSH("Common/CheckBox_Hovered", Icon16x16));
+
+		StyleSet->Set("EventFilter.State.Pending", new IMAGE_BRUSH("Common/CheckBox_Undetermined", Icon16x16));
+		StyleSet->Set("EventFilter.State.Pending_Hovered", new IMAGE_BRUSH("Common/CheckBox_Undetermined_Hovered", Icon16x16));
+
+		StyleSet->Set("EventFilter.TabIcon", new IMAGE_BRUSH("/Icons/icon_Genericfinder_16x", Icon16x16));
+	}
+
+
+	// Filter list
+	/* Set images for various SCheckBox states associated with "ContentBrowser.FilterButton" ... */
+	const FCheckBoxStyle FilterButtonCheckBoxStyle = FCheckBoxStyle()
+		.SetUncheckedImage(IMAGE_BRUSH("ContentBrowser/FilterUnchecked", FVector2D(10.0f, 20.0f)))
+		.SetUncheckedHoveredImage(IMAGE_BRUSH("ContentBrowser/FilterUnchecked", FVector2D(10.0f, 20.0f), FLinearColor(0.5f, 0.5f, 0.5f, 1.0f)))
+		.SetUncheckedPressedImage(IMAGE_BRUSH("ContentBrowser/FilterUnchecked", FVector2D(10.0f, 20.0f), FLinearColor(0.5f, 0.5f, 0.5f, 1.0f)))
+		.SetCheckedImage(IMAGE_BRUSH("ContentBrowser/FilterChecked", FVector2D(10.0f, 20.0f)))
+		.SetCheckedHoveredImage(IMAGE_BRUSH("ContentBrowser/FilterChecked", FVector2D(10.0f, 20.0f), FLinearColor(0.5f, 0.5f, 0.5f, 1.0f)))
+		.SetCheckedPressedImage(IMAGE_BRUSH("ContentBrowser/FilterChecked", FVector2D(10.0f, 20.0f), FLinearColor(0.5f, 0.5f, 0.5f, 1.0f)));
+	/* ... and add the new style */
+	StyleSet->Set("FilterPresets.FilterButton", FilterButtonCheckBoxStyle);
+
+	StyleSet->Set("FilterPresets.FilterNameFont", DEFAULT_FONT("Regular", 10));
+	StyleSet->Set("FilterPresets.FilterButtonBorder", new BOX_BRUSH("Common/RoundedSelection_16x", FMargin(4.0f / 16.0f)));
+	
+	FSlateStyleRegistry::RegisterSlateStyle( *StyleSet.Get() );
+}
+
+#undef IMAGE_BRUSH
+#undef BOX_BRUSH
+#undef BORDER_BRUSH
+#undef DEFAULT_FONT
+
+void FEventFilterStyle::Shutdown()
+{
+	if( StyleSet.IsValid() )
+	{
+		FSlateStyleRegistry::UnRegisterSlateStyle( *StyleSet.Get() );
+		ensure( StyleSet.IsUnique() );
+		StyleSet.Reset();
+	}
+}
+
+const ISlateStyle& FEventFilterStyle::Get()
+{
+	return *( StyleSet.Get() );
+}
+
+const FName& FEventFilterStyle::GetStyleSetName()
+{
+	return StyleSet->GetStyleSetName();
+}
diff --git a/Engine/Plugins/Developer/TraceDataFiltering/Source/TraceDataFiltering/Private/EventFilterStyle.h b/Engine/Plugins/Developer/TraceDataFiltering/Source/TraceDataFiltering/Private/EventFilterStyle.h
new file mode 100644
--- /dev/null
+++ b/Engine/Plugins/Developer/TraceDataFiltering/Source/TraceDataFiltering/Private/EventFilterStyle.h
@@ -0,0 +1,25 @@
+// Copyright Epic Games, Inc. All Rights Reserved.
+
+#pragma once
+
+#include "CoreMinimal.h"
+#include "Styling/SlateStyle.h"
+
+class FEventFilterStyle
+{
+public:
+	static void Initialize();
+	static void Shutdown();
+
+	static const ISlateStyle& Get();
+	static const FName& GetStyleSetName();
+	static const FSlateBrush* GetBrush(FName PropertyName, const ANSICHAR* Specifier = NULL)
+	{
+		return StyleSet->GetBrush(PropertyName, Specifier);
+	}
+
+private:
+
+	/** Singleton instances of this style. */
+	static TSharedPtr< class FSlateStyleSet > StyleSet;	
+};
diff --git a/Engine/Plugins/Developer/TraceDataFiltering/Source/TraceDataFiltering/Private/FilterPresets.cpp b/Engine/Plugins/Developer/TraceDataFiltering/Source/TraceDataFiltering/Private/FilterPresets.cpp
new file mode 100644
--- /dev/null
+++ b/Engine/Plugins/Developer/TraceDataFiltering/Source/TraceDataFiltering/Private/FilterPresets.cpp
@@ -0,0 +1,266 @@
+// Copyright Epic Games, Inc. All Rights Reserved.
+
+#include "FilterPresets.h"
+
+#define LOCTEXT_NAMESPACE "FilterPreset"
+
+void USharedFilterPresetContainer::GetSharedUserPresets(TArray<TSharedPtr<IFilterPreset>>& OutPresets)
+{
+	for (FFilterData& FilterData : SharedPresets)
+	{
+		OutPresets.Add(MakeShared<FUserFilterPreset>(FilterData.Name, FilterData));
+	}
+}
+
+void USharedFilterPresetContainer::AddFilterData(const FFilterData& InFilterData)
+{
+	USharedFilterPresetContainer* SharedPresetsContainer = GetMutableDefault<USharedFilterPresetContainer>();
+	SharedPresetsContainer->SharedPresets.Add(InFilterData);
+}
+
+bool USharedFilterPresetContainer::RemoveFilterData(const FFilterData& InFilterData)
+{
+	USharedFilterPresetContainer* SharedPresetsContainer = GetMutableDefault<USharedFilterPresetContainer>();
+	return SharedPresetsContainer->SharedPresets.RemoveSingle(InFilterData) == 1;
+}
+
+void USharedFilterPresetContainer::Save()
+{
+	USharedFilterPresetContainer* SharedPresetsContainer = GetMutableDefault<USharedFilterPresetContainer>();
+	SharedPresetsContainer->UpdateDefaultConfigFile();
+}
+
+void USharedFilterPresetContainer::GetEnginePresets(TArray<TSharedPtr<IFilterPreset>>& OutPresets)
+{
+	for (FFilterData& FilterData : EnginePresets)
+	{
+		OutPresets.Add(MakeShared<FEngineFilterPreset>(FilterData.Name, FilterData));
+	}
+}
+
+void ULocalFilterPresetContainer::GetUserPresets(TArray<TSharedPtr<IFilterPreset>>& OutPresets)
+{
+	for (FFilterData& FilterData : UserPresets)
+	{
+		const bool bIsLocal = true;
+		OutPresets.Add(MakeShared<FUserFilterPreset>(FilterData.Name, FilterData, bIsLocal));
+	}
+}
+
+void ULocalFilterPresetContainer::AddFilterData(const FFilterData& InFilterData)
+{
+	ULocalFilterPresetContainer* LocalPresetsContained = GetMutableDefault<ULocalFilterPresetContainer>();
+	LocalPresetsContained->UserPresets.Add(InFilterData);
+}
+
+bool ULocalFilterPresetContainer::RemoveFilterData(const FFilterData& InFilterData)
+{
+	ULocalFilterPresetContainer* LocalPresetsContained = GetMutableDefault<ULocalFilterPresetContainer>();
+	return LocalPresetsContained->UserPresets.RemoveSingle(InFilterData) == 1;
+}
+
+void ULocalFilterPresetContainer::Save()
+{
+	ULocalFilterPresetContainer* LocalPresetsContainer = GetMutableDefault<ULocalFilterPresetContainer>();
+	LocalPresetsContainer->SaveConfig();
+}
+
+void FFilterPresetHelpers::CreateNewPreset(const TArray<TSharedPtr<ITraceObject>>& InObjects)
+{
+	ULocalFilterPresetContainer* LocalPresetsContainer = GetMutableDefault<ULocalFilterPresetContainer>();
+	USharedFilterPresetContainer* SharedPresetsContainer = GetMutableDefault<USharedFilterPresetContainer>();
+
+	FFilterData& NewUserFilter = LocalPresetsContainer->UserPresets.AddDefaulted_GetRef();
+
+	bool bFoundValidName = false;
+	int32 ProfileAppendNum = LocalPresetsContainer->UserPresets.Num();
+	FString NewFilterName;
+	while (!bFoundValidName)
+	{
+		NewFilterName = FString::Printf(TEXT("UserPreset_%i"), ProfileAppendNum);
+
+		bool bValidName = true;
+		for (const FFilterData& Filter : LocalPresetsContainer->UserPresets)
+		{
+			if (Filter.Name == NewFilterName)
+			{
+				bValidName = false;
+				break;
+			}
+		}
+
+		for (const FFilterData& Filter : SharedPresetsContainer->SharedPresets)
+		{
+			if (Filter.Name == NewFilterName)
+			{
+				bValidName = false;
+				break;
+			}
+		}
+
+		if (!bValidName)
+		{
+			++ProfileAppendNum;
+		}
+
+		bFoundValidName = bValidName;
+	}
+
+	NewUserFilter.Name = NewFilterName;
+
+	TArray<FString> Names;
+	ExtractEnabledObjectNames(InObjects, Names);
+	NewUserFilter.WhitelistedNames = Names;
+
+	ULocalFilterPresetContainer::Save();
+}
+
+bool FFilterPresetHelpers::CanModifySharedPreset()
+{
+	const USharedFilterPresetContainer* SharedContainer = GetDefault<USharedFilterPresetContainer>();
+	return !FPlatformFileManager::Get().GetPlatformFile().IsReadOnly(*SharedContainer->GetDefaultConfigFilename());
+}
+
+void FFilterPresetHelpers::ExtractEnabledObjectNames(const TArray<TSharedPtr<ITraceObject>>& InObjects, TArray<FString>& OutNames)
+{
+	for (const TSharedPtr<ITraceObject>& Object : InObjects)
+	{
+		if (!Object->IsFiltered())
+		{
+			OutNames.Add(Object->GetName());
+		}
+
+		TArray<TSharedPtr<ITraceObject>> ChildObjects;
+		Object->GetChildren(ChildObjects);
+		if(ChildObjects.Num())
+		{
+			ExtractEnabledObjectNames(ChildObjects, OutNames);
+		}
+	}
+}
+
+bool FUserFilterPreset::CanDelete() const
+{
+	return true;
+}
+
+bool FUserFilterPreset::Delete()
+{
+	bool bRemoved = false;
+
+	if (IsLocal())
+	{
+		bRemoved = ULocalFilterPresetContainer::RemoveFilterData(FilterData);
+		ULocalFilterPresetContainer::Save();
+	}
+	else
+	{
+		bRemoved = USharedFilterPresetContainer::RemoveFilterData(FilterData);
+		USharedFilterPresetContainer::Save();
+	}
+
+	return ensure(bRemoved);
+}
+
+bool FUserFilterPreset::MakeShared()
+{
+	ensure(IsLocal());
+
+	USharedFilterPresetContainer::AddFilterData(FilterData);
+	ensure(ULocalFilterPresetContainer::RemoveFilterData(FilterData));
+		
+	USharedFilterPresetContainer::Save();
+	ULocalFilterPresetContainer::Save();
+	
+	return true;
+}
+
+bool FUserFilterPreset::MakeLocal()
+{
+	ensure(!IsLocal());
+
+	ULocalFilterPresetContainer::AddFilterData(FilterData);
+	ensure(USharedFilterPresetContainer::RemoveFilterData(FilterData));
+	
+	USharedFilterPresetContainer::Save();
+	ULocalFilterPresetContainer::Save();
+
+	return true;
+}
+
+bool FUserFilterPreset::IsLocal() const
+{
+	return bIsLocalPreset;
+}
+
+void FUserFilterPreset::Save(const TArray<TSharedPtr<ITraceObject>>& InObjects)
+{
+	TArray<FString> Names;
+	FFilterPresetHelpers::ExtractEnabledObjectNames(InObjects, Names);
+
+	FilterData.WhitelistedNames = Names;
+	
+	USharedFilterPresetContainer::Save();
+	ULocalFilterPresetContainer::Save();
+}
+
+void FUserFilterPreset::Save()
+{
+	USharedFilterPresetContainer::Save();
+	ULocalFilterPresetContainer::Save();
+}
+
+FString FFilterPreset::GetName() const
+{
+	return Name;
+}
+
+FText FFilterPreset::GetDisplayText() const
+{
+	return FText::FromString(Name);
+}
+
+FText FFilterPreset::GetDescription() const
+{
+	return FText::FormatOrdered(LOCTEXT("FilterPresetDescriptionFormat", "Name: {0}\nType: {1}"), FText::FromString(Name), CanDelete() ? (IsLocal() ? LOCTEXT("LocalPreset", "Local") : LOCTEXT("SharedPreset", "Shared")) : LOCTEXT("EnginePreset", "Engine"));
+}
+
+void FFilterPreset::GetWhitelistedNames(TArray<FString>& OutNames) const
+{
+	OutNames.Append(FilterData.WhitelistedNames);
+}
+
+bool FFilterPreset::CanDelete() const
+{
+	return false;
+}
+
+void FFilterPreset::Rename(const FString& InNewName)
+{	
+	Name = InNewName;
+	FilterData.Name = InNewName;
+
+	Save();
+}
+
+bool FFilterPreset::Delete()
+{
+	return false;
+}
+
+bool FFilterPreset::MakeShared()
+{
+	return false;
+}
+
+bool FFilterPreset::MakeLocal()
+{
+	return false;
+}
+
+bool FFilterPreset::IsLocal() const
+{
+	return false;
+}
+
+#undef LOCTEXT_NAMESPACE // "FilterPreset"
\ No newline at end of file
diff --git a/Engine/Plugins/Developer/TraceDataFiltering/Source/TraceDataFiltering/Private/FilterPresets.h b/Engine/Plugins/Developer/TraceDataFiltering/Source/TraceDataFiltering/Private/FilterPresets.h
new file mode 100644
--- /dev/null
+++ b/Engine/Plugins/Developer/TraceDataFiltering/Source/TraceDataFiltering/Private/FilterPresets.h
@@ -0,0 +1,126 @@
+// Copyright Epic Games, Inc. All Rights Reserved.
+
+#pragma once
+
+#include "IFilterPreset.h"
+#include "UObject/ObjectMacros.h"
+#include "Containers/Map.h"
+#include "Containers/UnrealString.h"
+#include "EventFilterService.h"
+#include "ITraceObject.h"
+
+#include "FilterPresets.generated.h"
+
+/** Structure representing an individual preset in configuration (ini) files */
+USTRUCT()
+struct FFilterData
+{
+	GENERATED_BODY()
+
+	UPROPERTY()
+	FString Name;
+
+	UPROPERTY()
+	TArray<FString> WhitelistedNames;
+	
+	bool operator==(const FFilterData& Other) const
+	{
+		return Name == Other.Name && WhitelistedNames == Other.WhitelistedNames;
+	}
+};
+
+struct FFilterPresetHelpers
+{
+	/** Creates a new filtering preset according to the specific object names */
+	static void CreateNewPreset(const TArray<TSharedPtr<ITraceObject>>& InObjects);
+	/** Creates a set of strings, corresponding to set of non-filtered out object as part of InObjects */
+	static void ExtractEnabledObjectNames(const TArray<TSharedPtr<ITraceObject>>& InObjects, TArray<FString>& OutNames);
+	/** Returns whether or not shared presets can be modified, requires write-flag on default confing files */
+	static bool CanModifySharedPreset();
+};
+
+/** Base implementation of a filter preset */
+struct FFilterPreset : public IFilterPreset
+{
+public:
+	FFilterPreset(const FString& InName, FFilterData& InFilterData) : Name(InName), FilterData(InFilterData) {}
+
+	/** Begin IFilterPreset overrides */
+	virtual FString GetName() const override;
+	virtual FText GetDisplayText() const;
+	virtual FText GetDescription() const;
+	virtual void GetWhitelistedNames(TArray<FString>& OutNames) const override;
+	virtual bool CanDelete() const override;
+	virtual void Rename(const FString& InNewName) override;
+	virtual bool Delete() override;
+	virtual bool MakeShared() override;
+	virtual bool MakeLocal() override;
+	virtual bool IsLocal() const override;
+	virtual void Save(const TArray<TSharedPtr<ITraceObject>>& InObjects) override {}
+	virtual void Save() override {}
+	/** End IFilterPreset overrides */
+protected:
+	FString Name;
+	FFilterData& FilterData;
+};
+
+/** Engine level preset is simply a basic one */
+typedef FFilterPreset FEngineFilterPreset;
+
+/** User filter preset, allows for deletion / transitioning INI ownership */
+struct FUserFilterPreset : public FFilterPreset
+{
+public:
+	FUserFilterPreset(const FString& InName, FFilterData& InFilterData, bool bInLocal = false) : FFilterPreset(InName, InFilterData), bIsLocalPreset(bInLocal) {}
+
+	/** Begin IFilterPreset overrides */
+	virtual bool CanDelete() const override;
+	virtual bool Delete() override;
+	virtual bool MakeShared() override;
+	virtual bool MakeLocal() override;
+	virtual bool IsLocal() const override;
+	virtual void Save(const TArray<TSharedPtr<ITraceObject>>& InObjects) override;
+	virtual void Save() override;
+	/** End IFilterPreset overrides */
+protected:
+	bool bIsLocalPreset;
+};
+
+/** UObject containers for the preset data */
+UCLASS(Config = TraceDataFilters)
+class ULocalFilterPresetContainer: public UObject
+{
+	GENERATED_BODY()
+
+	friend struct FFilterPresetHelpers;
+public:	
+	void GetUserPresets(TArray<TSharedPtr<IFilterPreset>>& OutPresets);
+
+	static void AddFilterData(const FFilterData& InFilterData);
+	static bool RemoveFilterData(const FFilterData& InFilterData);
+	static void Save();
+protected:	
+	UPROPERTY(Config)
+	TArray<FFilterData> UserPresets;
+};
+
+UCLASS(Config = TraceDataFilters, DefaultConfig)
+class USharedFilterPresetContainer : public UObject
+{
+	GENERATED_BODY()
+
+	friend struct FFilterPresetHelpers;
+public:
+	void GetEnginePresets(TArray<TSharedPtr<IFilterPreset>>& OutPresets);
+	void GetSharedUserPresets(TArray<TSharedPtr<IFilterPreset>>& OutPresets);
+
+	static void AddFilterData(const FFilterData& InFilterData);
+	static bool RemoveFilterData(const FFilterData& InFilterData);
+	static void Save();
+protected:
+	UPROPERTY(Config)
+	TArray<FFilterData> SharedPresets;
+
+	UPROPERTY(Config)
+	TArray<FFilterData> EnginePresets;
+};
\ No newline at end of file
diff --git a/Engine/Plugins/Developer/TraceDataFiltering/Source/TraceDataFiltering/Private/IFilterPreset.h b/Engine/Plugins/Developer/TraceDataFiltering/Source/TraceDataFiltering/Private/IFilterPreset.h
new file mode 100644
--- /dev/null
+++ b/Engine/Plugins/Developer/TraceDataFiltering/Source/TraceDataFiltering/Private/IFilterPreset.h
@@ -0,0 +1,36 @@
+// Copyright Epic Games, Inc. All Rights Reserved.
+
+#pragma once
+
+class ITraceObject;
+
+/** Interface for structures representing a filtering preset */
+struct IFilterPreset : public TSharedFromThis<IFilterPreset>
+{
+public:
+	virtual ~IFilterPreset() {}
+
+	virtual FString GetName() const = 0;
+	virtual FText GetDisplayText() const = 0;
+	virtual FText GetDescription() const = 0;
+
+	/** Returns all names which pass this preset its filter */
+	virtual void GetWhitelistedNames(TArray<FString>& OutNames) const = 0;
+	/** Whether or not this preset can be deleted */
+	virtual bool CanDelete() const = 0;
+	/** Whether or not this preset only exists locally */
+	virtual bool IsLocal() const = 0;
+
+	/** Update this preset according to the provided set of trace objects */
+	virtual void Save(const TArray<TSharedPtr<ITraceObject>>& InObjects) = 0;
+
+	virtual void Save() = 0;
+
+	virtual void Rename(const FString& InNewName) = 0;	
+	virtual bool Delete() = 0;
+
+	/** Transition preset between being local and shared, moves between local and default INI files */
+	virtual bool MakeShared() = 0;
+	virtual bool MakeLocal() = 0;
+	
+};
\ No newline at end of file
diff --git a/Engine/Plugins/Developer/TraceDataFiltering/Source/TraceDataFiltering/Private/ISessionTraceFilterService.h b/Engine/Plugins/Developer/TraceDataFiltering/Source/TraceDataFiltering/Private/ISessionTraceFilterService.h
new file mode 100644
--- /dev/null
+++ b/Engine/Plugins/Developer/TraceDataFiltering/Source/TraceDataFiltering/Private/ISessionTraceFilterService.h
@@ -0,0 +1,38 @@
+// Copyright Epic Games, Inc. All Rights Reserved.
+
+#pragma once
+
+#include "Containers/Array.h"
+#include "Containers/UnrealString.h"
+
+struct FTraceObjectInfo
+{
+	FString Name;
+	bool bEnabled;
+	uint32 Hash;
+	uint32 OwnerHash;
+};
+
+struct IFilterPreset;
+
+/** Filtering service, representing the state and data for a specific Trace::IAnalysisSession */
+class ISessionTraceFilterService : public TSharedFromThis<ISessionTraceFilterService>
+{
+public:
+	virtual ~ISessionTraceFilterService() {}
+
+	/** Returns the root level set of objects */
+	virtual void GetRootObjects(TArray<FTraceObjectInfo>& OutObjects) const = 0;
+
+	/** Returns the contained objects for the specific object hash */
+	virtual void GetChildObjects(uint32 InObjectHash, TArray<FTraceObjectInfo>& OutChildObjects) const = 0;
+
+	/** Set the filtered state for an individual object by its hash */
+	virtual void SetObjectFilterState(const FString& InObjectName, const bool bFilterState) = 0;
+	
+	/** Set timestamp for last processed update (data change) */
+	virtual const FDateTime& GetTimestamp() = 0;
+
+	/** Update filtering state according to user-set preset(s) */
+	virtual void UpdateFilterPresets(const TArray<TSharedPtr<IFilterPreset>>& InPresets) = 0;
+};
\ No newline at end of file
diff --git a/Engine/Plugins/Developer/TraceDataFiltering/Source/TraceDataFiltering/Private/ITraceObject.h b/Engine/Plugins/Developer/TraceDataFiltering/Source/TraceDataFiltering/Private/ITraceObject.h
new file mode 100644
--- /dev/null
+++ b/Engine/Plugins/Developer/TraceDataFiltering/Source/TraceDataFiltering/Private/ITraceObject.h
@@ -0,0 +1,35 @@
+// Copyright Epic Games, Inc. All Rights Reserved.
+
+#pragma once
+
+#include "Templates/SharedPointer.h"
+#include "Containers/Array.h"
+
+class ITraceObject : public TSharedFromThis<ITraceObject>
+{
+public:
+	virtual ~ITraceObject() {}
+
+	/** Returns the name representing this object */
+	virtual FString GetName() const = 0;
+	/** Returns the user-facing text for this object */
+	virtual FText GetDisplayText() const = 0;
+
+	/** Set whether or not this object should (not) be filtered out */
+	virtual void SetIsFiltered(bool bState) = 0;
+	
+	/** Set whether or not this object is pending an update of its state (from the connect edTrace connection) */
+	virtual void SetPending() = 0;
+	
+	/** Returns whether or not this object is set to be filtered out (true = filtered; false = not filtered)*/
+	virtual bool IsFiltered() const = 0;
+
+	/** Returns whether or not this object is pending a filter-state update (SetPending previously called) */
+	virtual bool IsPending() const = 0;
+
+	/** Returns contained child object(s) */
+	virtual void GetChildren(TArray<TSharedPtr<ITraceObject>>& OutChildren) const = 0;
+
+	/** Returns string representing this object in a search (treeview filtering) */
+	virtual void GetSearchString(TArray<FString>& OutFilterStrings) const = 0;
+};
\ No newline at end of file
diff --git a/Engine/Plugins/Developer/TraceDataFiltering/Source/TraceDataFiltering/Private/SFilterPreset.cpp b/Engine/Plugins/Developer/TraceDataFiltering/Source/TraceDataFiltering/Private/SFilterPreset.cpp
new file mode 100644
--- /dev/null
+++ b/Engine/Plugins/Developer/TraceDataFiltering/Source/TraceDataFiltering/Private/SFilterPreset.cpp
@@ -0,0 +1,318 @@
+// Copyright Epic Games, Inc. All Rights Reserved.
+
+#include "SFilterPreset.h"
+#include "EditorStyleSet.h"
+
+#define LOCTEXT_NAMESPACE "SFilterPreset"
+
+SFilterPreset::~SFilterPreset()
+{
+	if (IsHovered())
+	{
+		OnHighlightPreset.ExecuteIfBound(nullptr);
+	}
+}
+
+void SFilterPreset::Construct(const FArguments& InArgs)
+{
+	bEnabled = false;
+	bHighlighted = false;
+	OnPresetsChanged = InArgs._OnPresetsChanged;
+	OnRequestRemove = InArgs._OnRequestRemove;
+	OnRequestEnableOnly = InArgs._OnRequestEnableOnly;
+	OnRequestEnableAll = InArgs._OnRequestEnableAll;
+	OnRequestDisableAll = InArgs._OnRequestDisableAll;
+	OnRequestRemoveAll = InArgs._OnRequestRemoveAll;
+	OnRequestDelete = InArgs._OnRequestDelete;
+	OnRequestSave = InArgs._OnRequestSave;
+	OnHighlightPreset = InArgs._OnHighlightPreset;
+
+	FilterPreset = InArgs._FilterPreset;
+
+	const FName ColorName = [this]()
+	{
+		if (FilterPreset->CanDelete())
+		{
+			// User preset 
+			if (FilterPreset->IsLocal())
+			{
+				return FName("EventFilter.LocalPreset");
+			}
+			
+			return FName("EventFilter.SharedPreset");
+		}
+		
+		// Engine preset
+		return FName("EventFilter.EnginePreset");
+	}();
+
+	FilterColor = FEventFilterStyle::Get().GetColor(ColorName);
+
+	ChildSlot
+	[
+		SNew(SBorder)
+		.Padding(0)
+		.BorderBackgroundColor(FLinearColor(0.2f, 0.2f, 0.2f, 0.2f))
+		.BorderImage(FEventFilterStyle::GetBrush("FilterPresets.FilterButtonBorder"))
+		[
+			SAssignNew(ToggleButtonPtr, SFilterPresetCheckBox)
+			.Style(FEventFilterStyle::Get(), "FilterPresets.FilterButton")
+			.Padding(this, &SFilterPreset::GetNameLabelPadding)
+			.IsChecked(this, &SFilterPreset::IsChecked)
+			.OnCheckStateChanged(this, &SFilterPreset::PresetToggled)
+			.OnGetMenuContent(this, &SFilterPreset::GetRightClickMenuContent)
+			.ForegroundColor(this, &SFilterPreset::GetPresetForegroundColor)
+			[
+				SNew(STextBlock)
+				.ColorAndOpacity(this, &SFilterPreset::GetNameLabelColorAndOpacity)
+				.Font(FEventFilterStyle::Get().GetFontStyle("FilterPresets.FilterNameFont"))
+				.ShadowOffset(FVector2D(1.f, 1.f))
+				.Text(this, &SFilterPreset::GetPresetName)
+			]
+		]
+	];
+
+	ToggleButtonPtr->SetOnFilterCtrlClicked(FOnClicked::CreateSP(this, &SFilterPreset::FilterCtrlClicked));
+	ToggleButtonPtr->SetOnFilterAltClicked(FOnClicked::CreateSP(this, &SFilterPreset::FilterAltClicked));
+	ToggleButtonPtr->SetOnFilterDoubleClicked(FOnClicked::CreateSP(this, &SFilterPreset::FilterDoubleClicked));
+	ToggleButtonPtr->SetOnFilterMiddleButtonClicked(FOnClicked::CreateSP(this, &SFilterPreset::FilterMiddleButtonClicked));
+
+	TAttribute<FText> Attribute = TAttribute<FText>::Create(TAttribute<FText>::FGetter::CreateSP(this, &SFilterPreset::GetToolTipText));
+	ToggleButtonPtr->SetToolTipText(Attribute);
+}
+
+void SFilterPreset::SetEnabled(bool InEnabled, bool InExecuteOnPresetsChanged /*= true*/)
+{
+	if (InEnabled != bEnabled)
+	{
+		bEnabled = InEnabled;
+		if (InExecuteOnPresetsChanged)
+		{
+			OnPresetsChanged.ExecuteIfBound();
+		}
+	}
+}
+
+bool SFilterPreset::IsEnabled() const
+{
+	return bEnabled;
+}
+
+const TSharedPtr<IFilterPreset>& SFilterPreset::GetFilterPreset() const
+{
+	return FilterPreset;
+}
+
+void SFilterPreset::OnMouseLeave(const FPointerEvent& MouseEvent)
+{
+	SCompoundWidget::OnMouseLeave(MouseEvent);
+	
+	if (bHighlighted)
+	{
+		OnHighlightPreset.ExecuteIfBound(nullptr);
+	}
+	
+	bHighlighted = false;
+}
+
+FReply SFilterPreset::OnMouseMove(const FGeometry& MyGeometry, const FPointerEvent& MouseEvent)
+{
+	if (IsHovered())
+	{
+		const bool bWasHighlighted = bHighlighted;
+		bHighlighted = MouseEvent.IsControlDown();
+		if (bHighlighted || bWasHighlighted)
+		{
+			OnHighlightPreset.ExecuteIfBound(bHighlighted ? FilterPreset : nullptr);
+		}
+	}
+
+	return SCompoundWidget::OnMouseMove(MyGeometry, MouseEvent);
+}
+
+void SFilterPreset::PresetToggled(ECheckBoxState NewState)
+{
+	bEnabled = NewState == ECheckBoxState::Checked;
+	OnPresetsChanged.ExecuteIfBound();
+}
+
+FReply SFilterPreset::FilterCtrlClicked()
+{
+	OnRequestEnableAll.ExecuteIfBound();
+	return FReply::Handled();
+}
+
+FReply SFilterPreset::FilterAltClicked()
+{
+	OnRequestDisableAll.ExecuteIfBound();
+	return FReply::Handled();
+}
+
+FReply SFilterPreset::FilterDoubleClicked()
+{
+	// Disable all other presets and enable this one.
+	OnRequestDisableAll.ExecuteIfBound();
+	bEnabled = true;
+	OnPresetsChanged.ExecuteIfBound();
+
+	return FReply::Handled();
+}
+
+FReply SFilterPreset::FilterMiddleButtonClicked()
+{
+	RemovePreset();
+	return FReply::Handled();
+}
+
+FText SFilterPreset::GetToolTipText() const
+{
+	return bHighlighted ? FilterPreset->GetDescription() : FText::Format(LOCTEXT("FilterPresetTooltipText", "{0}\nHold CTRL to highlight whitelisted items."), FilterPreset->GetDescription());
+}
+
+TSharedRef<SWidget> SFilterPreset::GetRightClickMenuContent()
+{
+	FMenuBuilder MenuBuilder(/*bInShouldCloseWindowAfterMenuSelection=*/true, NULL);
+
+	MenuBuilder.BeginSection("PresetOptions", LOCTEXT("FilterContextHeading", "Preset Options"));
+	{
+		MenuBuilder.AddMenuEntry(
+			FText::Format(LOCTEXT("RemovePreset", "Remove: {0}"), GetPresetName()),
+			LOCTEXT("RemovePresetTooltip", "Remove this preset from the list. It can be added again in the presets menu."),
+			FSlateIcon(),
+			FUIAction(FExecuteAction::CreateSP(this, &SFilterPreset::RemovePreset))
+		);
+
+		MenuBuilder.AddMenuEntry(
+			FText::Format(LOCTEXT("EnableOnlyThisPreset", "Enable this only: {0}"), GetPresetName()),
+			LOCTEXT("EnableOnlyThisPresetTooltip", "Enable only this prest from the list."),
+			FSlateIcon(),
+			FUIAction(FExecuteAction::CreateSP(this, &SFilterPreset::EnableOnly))
+		);
+
+	}
+	MenuBuilder.EndSection();
+
+	MenuBuilder.BeginSection("PresetBulkOptions", LOCTEXT("BulkFilterContextHeading", "Bulk Preset Options"));
+	{
+		MenuBuilder.AddMenuEntry(
+			LOCTEXT("EnableAllPresets", "Enable All Presets"),
+			LOCTEXT("EnableAllPresetsTooltip", "Enables all presets."),
+			FSlateIcon(),
+			FUIAction(FExecuteAction::CreateSP(this, &SFilterPreset::EnableAllPresets))
+		);
+
+		MenuBuilder.AddMenuEntry(
+			LOCTEXT("DisableAllPresets", "Disable All Presets"),
+			LOCTEXT("DisableAllPresetsTooltip", "Disables all active presets."),
+			FSlateIcon(),
+			FUIAction(FExecuteAction::CreateSP(this, &SFilterPreset::DisableAllPresets))
+		);
+
+		MenuBuilder.AddMenuEntry(
+			LOCTEXT("RemoveAllPresets", "Remove All Presets"),
+			LOCTEXT("RemoveAllPresetsTooltip", "Removes all presets from the list."),
+			FSlateIcon(),
+			FUIAction(FExecuteAction::CreateSP(this, &SFilterPreset::RemoveAllPresets))
+		);
+	}
+	MenuBuilder.EndSection();
+
+
+	MenuBuilder.BeginSection("UserPresetOptions", LOCTEXT("UserPresetsContextHeading", "User Presets"));
+	{
+		if (FilterPreset->CanDelete())
+		{
+			MenuBuilder.AddMenuEntry(
+				FText::Format(LOCTEXT("DeletePreset", "Delete User Preset: {0}"), GetPresetName()),
+				LOCTEXT("DeleteFilterTooltip", "Deletes this User Preset."),
+				FSlateIcon(),
+				FUIAction(FExecuteAction::CreateSP(this, &SFilterPreset::DeletePreset))
+			);
+
+			MenuBuilder.AddMenuEntry(
+				FText::Format(LOCTEXT("SavePreset", "Save User Preset: {0}"), GetPresetName()),
+				LOCTEXT("SaveFilterTooltip", "Saves the currently filtering state as the selected User Preset."),
+				FSlateIcon(),
+				FUIAction(FExecuteAction::CreateSP(this, &SFilterPreset::SavePreset))
+			);
+		}
+	}
+	MenuBuilder.EndSection();
+
+
+	return MenuBuilder.MakeWidget();
+}
+
+void SFilterPreset::RemovePreset()
+{
+	TSharedRef<SFilterPreset> Self = SharedThis(this);
+	OnRequestRemove.ExecuteIfBound(Self);
+}
+
+void SFilterPreset::EnableOnly()
+{
+	TSharedRef<SFilterPreset> Self = SharedThis(this);
+	OnRequestEnableOnly.ExecuteIfBound(Self);
+}
+
+void SFilterPreset::EnableAllPresets()
+{
+	OnRequestEnableAll.ExecuteIfBound();
+}
+
+void SFilterPreset::DisableAllPresets()
+{
+	OnRequestDisableAll.ExecuteIfBound();
+}
+
+void SFilterPreset::RemoveAllPresets()
+{
+	OnRequestRemoveAll.ExecuteIfBound();
+}
+
+void SFilterPreset::SavePreset()
+{
+	TSharedRef<SFilterPreset> Self = SharedThis(this);
+	OnRequestSave.ExecuteIfBound(Self);
+}
+
+void SFilterPreset::DeletePreset()
+{
+	TSharedRef<SFilterPreset> Self = SharedThis(this);
+	OnRequestDelete.ExecuteIfBound(Self);
+}
+
+ECheckBoxState SFilterPreset::IsChecked() const
+{
+	return bEnabled ? ECheckBoxState::Checked : ECheckBoxState::Unchecked;
+}
+
+FSlateColor SFilterPreset::GetPresetForegroundColor() const
+{
+	return IsChecked() == ECheckBoxState::Checked ? FilterColor : FLinearColor::White;
+}
+
+FMargin SFilterPreset::GetNameLabelPadding() const
+{
+	return ToggleButtonPtr->IsPressed() ? FMargin(3, 2, 4, 0) : FMargin(3, 1, 4, 1);
+}
+
+FSlateColor SFilterPreset::GetNameLabelColorAndOpacity() const
+{
+	const float DimFactor = 0.75f;
+	return IsHovered() ? FLinearColor(DimFactor, DimFactor, DimFactor, 1.0f) : FLinearColor::White;
+}
+
+FText SFilterPreset::GetPresetName() const
+{
+	FText PresetDisplayText = FilterPreset->GetDisplayText();
+
+	if (PresetDisplayText.IsEmpty())
+	{
+		PresetDisplayText = LOCTEXT("UnknownPresetName", "???");
+	}
+
+	return PresetDisplayText;
+}
+
+#undef LOCTEXT_NAMESPACE // "SFilterPreset"
\ No newline at end of file
diff --git a/Engine/Plugins/Developer/TraceDataFiltering/Source/TraceDataFiltering/Private/SFilterPreset.h b/Engine/Plugins/Developer/TraceDataFiltering/Source/TraceDataFiltering/Private/SFilterPreset.h
new file mode 100644
--- /dev/null
+++ b/Engine/Plugins/Developer/TraceDataFiltering/Source/TraceDataFiltering/Private/SFilterPreset.h
@@ -0,0 +1,222 @@
+// Copyright Epic Games, Inc. All Rights Reserved.
+
+#pragma once
+
+#include "CoreMinimal.h"
+
+DECLARE_DELEGATE(FOnPresetsChanged);
+
+/** A class for check boxes in the preset list. If you double click the checkbox, you will enable it and disable all others */
+class SFilterPresetCheckBox : public SCheckBox
+{
+public:
+	
+	void SetOnFilterCtrlClicked(const FOnClicked& NewFilterCtrlClicked)
+	{
+		OnFilterCtrlClicked = NewFilterCtrlClicked;
+	}
+
+	void SetOnFilterAltClicked(const FOnClicked& NewFilteAltClicked)
+	{
+		OnFilterAltClicked = NewFilteAltClicked;
+	}
+
+	void SetOnFilterDoubleClicked(const FOnClicked& NewFilterDoubleClicked)
+	{
+		OnFilterDoubleClicked = NewFilterDoubleClicked;
+	}
+
+	void SetOnFilterMiddleButtonClicked(const FOnClicked& NewFilterMiddleButtonClicked)
+	{
+		OnFilterMiddleButtonClicked = NewFilterMiddleButtonClicked;
+	}
+
+	virtual FReply OnMouseButtonDoubleClick(const FGeometry& InMyGeometry, const FPointerEvent& InMouseEvent) override
+	{
+		// Handle left mouse button double click
+		if (InMouseEvent.GetEffectingButton() == EKeys::LeftMouseButton && OnFilterDoubleClicked.IsBound())
+		{
+			return OnFilterDoubleClicked.Execute();
+		}
+		else
+		{
+			return SCheckBox::OnMouseButtonDoubleClick(InMyGeometry, InMouseEvent);
+		}
+	}
+
+	virtual FReply OnMouseButtonUp(const FGeometry& InMyGeometry, const FPointerEvent& InMouseEvent) override
+	{
+		SCheckBox::OnMouseButtonUp(InMyGeometry, InMouseEvent);
+
+		if (InMouseEvent.IsControlDown() && OnFilterCtrlClicked.IsBound())
+		{
+			return OnFilterCtrlClicked.Execute();
+		}
+		else if (InMouseEvent.IsAltDown() && OnFilterAltClicked.IsBound())
+		{
+			return OnFilterAltClicked.Execute();
+		}
+		else if (InMouseEvent.GetEffectingButton() == EKeys::MiddleMouseButton && OnFilterMiddleButtonClicked.IsBound())
+		{
+			return OnFilterMiddleButtonClicked.Execute();
+		}
+		else
+		{			
+			return FReply::Handled().ReleaseMouseCapture();
+		}
+	}
+
+private:
+	FOnClicked OnFilterCtrlClicked;
+	FOnClicked OnFilterAltClicked;
+	FOnClicked OnFilterDoubleClicked;
+	FOnClicked OnFilterMiddleButtonClicked;
+};
+
+/**
+ * A single filtering preset in the preset list. Can be removed by clicking the remove button on it.
+ */
+class SFilterPreset : public SCompoundWidget
+{
+public:
+	DECLARE_DELEGATE_OneParam(FOnRequestRemove, const TSharedRef<SFilterPreset>& /*PresetToRemove*/);
+	DECLARE_DELEGATE_OneParam(FOnRequestEnableOnly, const TSharedRef<SFilterPreset>& /*PresetToEnable*/);
+	DECLARE_DELEGATE_OneParam(FOnRequestDelete, const TSharedRef<SFilterPreset>& /*PresetToDelete*/);
+	DECLARE_DELEGATE_OneParam(FOnRequestSave, const TSharedRef<SFilterPreset>& /*PresetToSave*/);
+	DECLARE_DELEGATE_OneParam(FOnHighlightPreset, const TSharedPtr<IFilterPreset>& /*PresetToHighlight*/);
+	DECLARE_DELEGATE(FOnRequestEnableAll);
+	DECLARE_DELEGATE(FOnRequestDisableAll);
+	DECLARE_DELEGATE(FOnRequestRemoveAll);	
+
+	SLATE_BEGIN_ARGS(SFilterPreset) {}
+		/** Filter preset this widget represents */
+		SLATE_ARGUMENT(TSharedPtr<IFilterPreset>, FilterPreset)
+
+		/** Invoked when the preset toggled */
+		SLATE_EVENT(FOnPresetsChanged, OnPresetsChanged)
+
+		/** Invoked when a request to remove this preset originated from within this widget */
+		SLATE_EVENT(FOnRequestRemove, OnRequestRemove)
+
+		/** Invoked when a request to enable only this preset originated from within this widget */
+		SLATE_EVENT(FOnRequestEnableOnly, OnRequestEnableOnly)
+
+		/** Invoked when a request to delete this preset originated from within this widget */
+		SLATE_EVENT(FOnRequestDelete, OnRequestDelete);
+
+		/** Invoked when a request to save the current filtering state as this preset originated from within this widget */
+		SLATE_EVENT(FOnRequestSave, OnRequestSave);
+
+		/** Invoked when a request to highlight this preset originated from within this widget */
+		SLATE_EVENT(FOnHighlightPreset, OnHighlightPreset)
+
+		/** Invoked when a request to enable all presets originated from within this widget */
+		SLATE_EVENT(FOnRequestEnableAll, OnRequestEnableAll)
+
+		/** Invoked when a request to disable all presets originated from within this widget */
+		SLATE_EVENT(FOnRequestDisableAll, OnRequestDisableAll)
+
+		/** Invoked when a request to remove all presets originated from within this widget */
+		SLATE_EVENT(FOnRequestRemoveAll, OnRequestRemoveAll)
+	SLATE_END_ARGS()
+
+	virtual ~SFilterPreset();
+
+	/** Constructs this widget with InArgs */
+	void Construct(const FArguments& InArgs);
+
+	/** Begin SCompoundWidget overrides */
+	virtual void OnMouseLeave(const FPointerEvent& MouseEvent) override;
+	virtual FReply OnMouseMove(const FGeometry& MyGeometry, const FPointerEvent& MouseEvent) override;
+	/** End SCompoundWidget overrides */
+
+	/** Sets whether or not this preset is applied to the combined filtering state */
+	void SetEnabled(bool InEnabled, bool InExecuteOnPresetsChanged = true);
+
+	/** Returns true if this preset contributes to the combined filtering state */
+	bool IsEnabled() const;
+
+	/** Get the preset object this widget represents */
+	const TSharedPtr<IFilterPreset>& GetFilterPreset() const;
+private:
+	/** Handler for when the preset checkbox is clicked */
+	void PresetToggled(ECheckBoxState NewState);
+
+	/** Handler for when the preset checkbox is clicked and a control key is pressed */
+	FReply FilterCtrlClicked();
+
+	/** Handler for when the preset checkbox is clicked and an alt key is pressed */
+	FReply FilterAltClicked();
+
+	/** Handler for when the preset checkbox is double clicked */
+	FReply FilterDoubleClicked();
+
+	/** Handler for when the preset checkbox is middle button clicked */
+	FReply FilterMiddleButtonClicked();
+
+	FText GetToolTipText() const;
+
+	/** Handler to create a right click menu */
+	TSharedRef<SWidget> GetRightClickMenuContent();
+
+	/** Removes this preset from the presets list */
+	void RemovePreset();
+
+	/** Enables only this preset from the presets list */
+	void EnableOnly();
+
+	/** Enables all presets in the list */
+	void EnableAllPresets();
+
+	/** Disables all active presets in the list */
+	void DisableAllPresets();
+
+	/** Removes all presets in the list */
+	void RemoveAllPresets();
+
+	/** Saves the current filtering state as this preset */
+	void SavePreset();
+
+	/** Deletes this preset */
+	void DeletePreset();
+
+	/** Handler to determine the "checked" state of the preset checkbox */
+	ECheckBoxState IsChecked() const;
+
+	/** Handler to determine the color of the checkbox when it is checked */
+	FSlateColor GetPresetForegroundColor() const;
+
+	/** Handler to determine the padding of the checkbox text when it is pressed */
+	FMargin GetNameLabelPadding() const;
+
+	/** Handler to determine the color of the checkbox text when it is hovered */
+	FSlateColor GetNameLabelColorAndOpacity() const;
+
+	/** Returns the display name for this filter */
+	FText GetPresetName() const;
+
+private:
+	TSharedPtr<IFilterPreset> FilterPreset;
+
+	FOnPresetsChanged OnPresetsChanged;
+	FOnRequestRemove OnRequestRemove;
+	FOnRequestDelete OnRequestDelete;
+	FOnRequestSave OnRequestSave;
+	FOnHighlightPreset OnHighlightPreset;
+	FOnRequestEnableOnly OnRequestEnableOnly;
+	FOnRequestEnableAll OnRequestEnableAll;
+	FOnRequestDisableAll OnRequestDisableAll;
+	FOnRequestDisableAll OnRequestRemoveAll;	
+
+	/** True when this preset should be applied to filtering state */
+	bool bEnabled;
+
+	/** The button to toggle the preset on or off */
+	TSharedPtr<SFilterPresetCheckBox> ToggleButtonPtr;
+
+	/** The color of the checkbox for this preset */
+	FLinearColor FilterColor;
+
+	/** True when this preset its filtering state is currently set to being highlighted */
+	bool bHighlighted;
+};
diff --git a/Engine/Plugins/Developer/TraceDataFiltering/Source/TraceDataFiltering/Private/SFilterPresetList.cpp b/Engine/Plugins/Developer/TraceDataFiltering/Source/TraceDataFiltering/Private/SFilterPresetList.cpp
new file mode 100644
--- /dev/null
+++ b/Engine/Plugins/Developer/TraceDataFiltering/Source/TraceDataFiltering/Private/SFilterPresetList.cpp
@@ -0,0 +1,563 @@
+// Copyright Epic Games, Inc. All Rights Reserved.
+
+#include "SFilterPresetList.h"
+
+#include "UObject/UObjectHash.h"
+#include "UObject/UObjectIterator.h"
+#include "Framework/Commands/UIAction.h"
+#include "Framework/Application/MenuStack.h"
+#include "Framework/Application/SlateApplication.h"
+#include "Framework/MultiBox/MultiBoxBuilder.h"
+#include "Widgets/SBoxPanel.h"
+#include "Widgets/Layout/SBorder.h"
+#include "Widgets/Layout/SWrapBox.h"
+#include "Widgets/Input/SEditableText.h"
+#include "Widgets/Text/STextBlock.h"
+#include "Widgets/Input/SCheckBox.h"
+#include "Widgets/Text/SInlineEditableTextBlock.h"
+#include "Styling/SlateTypes.h"
+#include "Misc/ConfigCacheIni.h"
+#include "Layout/WidgetPath.h"
+#include "Textures/SlateIcon.h"
+#include "EditorStyleSet.h"
+
+#include "TraceDataFilteringModule.h"
+
+#if WITH_EDITOR
+#include "SSettingsEditorCheckoutNotice.h"
+#endif
+
+#define LOCTEXT_NAMESPACE "InsightsFilterList"
+
+void SFilterPresetList::Construct( const FArguments& InArgs )
+{
+	OnPresetsChanged = InArgs._OnPresetsChanged;
+	OnSavePreset = InArgs._OnSavePreset;
+	OnHighlightPreset = InArgs._OnHighlightPreset;
+
+	FilterBox = SNew(SWrapBox)
+		.UseAllottedWidth(true);
+
+	LoadSettings(FTraceFilteringModule::TraceFiltersIni);
+	
+	PresetContainer = GetMutableDefault<ULocalFilterPresetContainer>();
+	SharedPresetContainer = GetMutableDefault<USharedFilterPresetContainer>();
+	
+	RefreshFilterPresets();
+
+	ChildSlot
+	[
+		FilterBox.ToSharedRef()
+	];
+
+}
+
+SFilterPresetList::~SFilterPresetList()
+{
+	SaveSettings(FTraceFilteringModule::TraceFiltersIni);
+}
+
+void SFilterPresetList::RefreshFilterPresets()
+{
+	FSlateApplication::Get().DismissAllMenus();
+
+	if (HasAnyPresets())
+	{
+		for (auto FilterIt = Presets.CreateConstIterator(); FilterIt; ++FilterIt)
+		{
+			const TSharedRef<SFilterPreset>& FilterToRemove = *FilterIt;
+			ActiveFilterNames.Add(FilterToRemove->GetFilterPreset()->GetName());
+
+			if (FilterToRemove->IsEnabled())
+			{
+				EnabledFilterNames.Add(FilterToRemove->GetFilterPreset()->GetName());
+			}
+		}
+
+		FilterBox->ClearChildren();
+		Presets.Empty();
+	}
+
+	Presets.Empty();
+	EngineFilterPresets.Empty();
+	UserFilterPresets.Empty();
+	SharedUserFilterPresets.Empty();
+	AllFilterPresets.Empty();
+
+	PresetContainer->GetUserPresets(UserFilterPresets);
+
+	SharedPresetContainer->GetEnginePresets(EngineFilterPresets);
+	SharedPresetContainer->GetSharedUserPresets(SharedUserFilterPresets);
+
+	AllFilterPresets.Append(EngineFilterPresets);
+	AllFilterPresets.Append(UserFilterPresets);
+	AllFilterPresets.Append(SharedUserFilterPresets);
+
+	for (const TSharedPtr<IFilterPreset>& Preset : AllFilterPresets)
+	{
+		if (ActiveFilterNames.Contains(Preset->GetName()))
+		{
+			TSharedRef<SFilterPreset> Filter = AddFilterPreset(Preset);
+			if (EnabledFilterNames.Contains(Preset->GetName()))
+			{
+				Filter->SetEnabled(true);
+			}
+		}
+	}
+
+	if (HasAnyPresets())
+	{
+		OnPresetsChanged.ExecuteIfBound();
+	}
+
+	ActiveFilterNames.Empty();
+	EnabledFilterNames.Empty();
+}
+
+FReply SFilterPresetList::OnMouseButtonUp( const FGeometry& MyGeometry, const FPointerEvent& MouseEvent )
+{
+	if ( MouseEvent.GetEffectingButton() == EKeys::RightMouseButton )
+	{
+		FReply Reply = FReply::Handled().ReleaseMouseCapture();
+
+		// Get the context menu content. If NULL, don't open a menu.
+		TSharedPtr<SWidget> MenuContent = MakeFilterPresetsMenu();
+
+		if ( MenuContent.IsValid() )
+		{
+			FVector2D SummonLocation = MouseEvent.GetScreenSpacePosition();
+			FWidgetPath WidgetPath = MouseEvent.GetEventPath() != nullptr ? *MouseEvent.GetEventPath() : FWidgetPath();
+			FSlateApplication::Get().PushMenu(AsShared(), WidgetPath, MenuContent.ToSharedRef(), SummonLocation, FPopupTransitionEffect(FPopupTransitionEffect::ContextMenu));
+		}
+
+		return Reply;
+	}
+
+	return FReply::Unhandled();
+}
+
+void SFilterPresetList::SaveSettings(const FString& IniFilename) const
+{
+	FString IniSectionName = TEXT("FilterPresetList");
+	FString IniActivePresetsKey = TEXT("ActivePresets");
+	FString IniEnabledPresetsKey = TEXT("EnabledPresets");
+
+	FString Separator = TEXT(";");
+	
+	FString ActivePresetsString;
+	FString EnabledPresetsString;
+
+	for (const TSharedRef<SFilterPreset>& Preset : Presets)
+	{
+		ActivePresetsString.Append(Preset->GetFilterPreset()->GetName());
+		ActivePresetsString.Append(Separator);
+
+		if (Preset->IsEnabled())
+		{
+			EnabledPresetsString.Append(Preset->GetFilterPreset()->GetName());
+			EnabledPresetsString.Append(Separator);
+		}
+	}
+
+	ActivePresetsString.RemoveFromEnd(Separator);
+	EnabledPresetsString.RemoveFromEnd(Separator);
+
+	GConfig->SetString(*IniSectionName, *IniActivePresetsKey, *ActivePresetsString, IniFilename);
+	GConfig->SetString(*IniSectionName, *IniEnabledPresetsKey, *EnabledPresetsString, IniFilename);
+}
+
+void SFilterPresetList::LoadSettings(const FString& IniFilename)
+{
+	FString IniSectionName = TEXT("FilterPresetList");
+	FString IniActivePresetsKey = TEXT("ActivePresets");
+	FString IniEnabledPresetsKey = TEXT("EnabledPresets");
+
+	FString ActivePresetsString;
+	FString EnabledPresetsString;
+	GConfig->GetString(*IniSectionName, *IniActivePresetsKey, ActivePresetsString, IniFilename);
+	GConfig->GetString(*IniSectionName, *IniEnabledPresetsKey, EnabledPresetsString, IniFilename);
+
+	FString Separator = TEXT(";");
+	ActivePresetsString.ParseIntoArray(ActiveFilterNames, *Separator);
+	EnabledPresetsString.ParseIntoArray(EnabledFilterNames, *Separator);
+}
+
+bool SFilterPresetList::HasAnyPresets() const
+{
+	return Presets.Num() > 0;
+}
+
+void SFilterPresetList::EnableAllPresets()
+{
+	for (auto FilterIt = Presets.CreateConstIterator(); FilterIt; ++FilterIt)
+	{
+		(*FilterIt)->SetEnabled(true);
+	}
+}
+
+void SFilterPresetList::DisableAllPresets()
+{
+	for ( auto FilterIt = Presets.CreateConstIterator(); FilterIt; ++FilterIt )
+	{
+		(*FilterIt)->SetEnabled(false, false);
+	}
+
+	OnPresetsChanged.ExecuteIfBound();
+}
+
+void SFilterPresetList::RemoveAllPresets()
+{
+	if ( HasAnyPresets() )
+	{
+		for ( auto FilterIt = Presets.CreateConstIterator(); FilterIt; ++FilterIt )
+		{
+			const TSharedRef<SFilterPreset>& FilterToRemove = *FilterIt;		
+		}
+
+		FilterBox->ClearChildren();
+		Presets.Empty();
+
+		// Notify that a filter has changed
+		OnPresetsChanged.ExecuteIfBound();
+	}
+}
+
+
+void SFilterPresetList::GenerateEnginePresetsMenu(FMenuBuilder& InMenuBuilder)
+{
+	InMenuBuilder.BeginSection("TracePresetsEnginePresets", LOCTEXT("EnginePresetsMenuHeading", "Engine Presets"));
+	{
+		for (const TSharedPtr<IFilterPreset>& EnginePreset : EngineFilterPresets)
+		{
+			InMenuBuilder.AddMenuEntry(
+				EnginePreset->GetDisplayText(),
+				EnginePreset->GetDisplayText(),
+				FSlateIcon(),
+				FUIAction(
+					FExecuteAction::CreateSP(this, &SFilterPresetList::TogglePreset, EnginePreset),
+					FCanExecuteAction(),
+					FIsActionChecked::CreateSP(this, &SFilterPresetList::IsPresetEnabled, EnginePreset)),
+				NAME_None,
+				EUserInterfaceActionType::ToggleButton
+			);
+		}
+	}
+	InMenuBuilder.EndSection();  // TracePresetsEnginePresets
+}
+
+void SFilterPresetList::GenerateLocalUserPresetsMenu(FMenuBuilder& InMenuBuilder)
+{
+	InMenuBuilder.BeginSection("TracePresetsUserPresets", LOCTEXT("UserPresetsMenuHeading", "User Presets"));
+	{
+		for (const TSharedPtr<IFilterPreset>& UserPreset : UserFilterPresets)
+		{
+			InMenuBuilder.AddSubMenu(
+				UserPreset->GetDisplayText(),
+				UserPreset->GetDisplayText(),
+				FNewMenuDelegate::CreateLambda([this, UserPreset](FMenuBuilder& InSubMenuBuilder)
+				{
+					InSubMenuBuilder.BeginSection(NAME_None, LOCTEXT("UserPresetItemsHead", "User Preset(s)"));
+					{
+						GenerateCommonPresetEntries(InSubMenuBuilder, UserPreset);
+					}
+					InSubMenuBuilder.EndSection();
+				}),
+				FUIAction(
+					FExecuteAction::CreateSP(this, &SFilterPresetList::TogglePreset, UserPreset),
+					FCanExecuteAction(),
+					FIsActionChecked::CreateSP(this, &SFilterPresetList::IsPresetEnabled, UserPreset)),
+				NAME_None,
+				EUserInterfaceActionType::ToggleButton,
+				false, 
+				FSlateIcon(),
+				false
+			);
+		}
+	}
+	InMenuBuilder.EndSection(); // TracePresetsUserPresets
+}
+
+void SFilterPresetList::GenerateSharedUserPresetsMenu(FMenuBuilder& InMenuBuilder)
+{
+	InMenuBuilder.BeginSection("TracePresetsSharedUserPresets", LOCTEXT("SharedUserPresetsMenuHeading", "Shared User Presets"));
+	{
+		for (const TSharedPtr<IFilterPreset>& UserPreset : SharedUserFilterPresets)
+		{
+			InMenuBuilder.AddSubMenu(
+				UserPreset->GetDisplayText(),
+				UserPreset->GetDisplayText(),
+				FNewMenuDelegate::CreateLambda([this, UserPreset](FMenuBuilder& InSubMenuBuilder)
+				{
+#if WITH_EDITOR
+					// SSettingsEditorCheckoutNotice is part of SharedWidgets module which depends on UnrealEd
+					TSharedPtr<SSettingsEditorCheckoutNotice> CheckoutWidget;
+					
+					SAssignNew(CheckoutWidget, SSettingsEditorCheckoutNotice)
+					.ConfigFilePath(GetDefault<USharedFilterPresetContainer>()->GetDefaultConfigFilename())
+					.Visibility_Lambda([]() -> EVisibility
+					{
+						return FFilterPresetHelpers::CanModifySharedPreset() ? EVisibility::Collapsed : EVisibility::Visible;
+					});
+#endif // WITH_EDITOR
+					InSubMenuBuilder.BeginSection(NAME_None, LOCTEXT("SharedUserPresetItemHead", "Shared User Preset(s)"));
+					{
+#if WITH_EDITOR
+						InSubMenuBuilder.AddWidget(CheckoutWidget.ToSharedRef(), FText::GetEmpty());
+#endif // WITH_EDITOR
+						GenerateCommonPresetEntries(InSubMenuBuilder, UserPreset);
+					}
+					InSubMenuBuilder.EndSection();
+				}),
+				FUIAction(
+					FExecuteAction::CreateSP(this, &SFilterPresetList::TogglePreset, UserPreset),
+					FCanExecuteAction(),
+					FIsActionChecked::CreateSP(this, &SFilterPresetList::IsPresetEnabled, UserPreset)),
+				NAME_None,
+				EUserInterfaceActionType::ToggleButton,
+				false, 
+				FSlateIcon(),
+				false
+			);
+		}
+	}
+	InMenuBuilder.EndSection(); // TracePresetsSharedUserPresets
+}
+
+void SFilterPresetList::GenerateCommonPresetEntries(FMenuBuilder& InSubMenuBuilder, const TSharedPtr<IFilterPreset>& UserPreset)
+{
+	InSubMenuBuilder.AddWidget(
+		SNew(SEditableTextBox)
+		.Text_Lambda([UserPreset]() { return UserPreset->GetDisplayText(); })
+		.OnTextCommitted_Lambda([UserPreset](const FText& InText, ETextCommit::Type InCommitType)
+		{
+			UserPreset->Rename(InText.ToString());
+		})
+		.OnVerifyTextChanged_Lambda([this, UserPreset](const FText& InNewText, FText& OutErrorMessage) -> bool
+		{
+			auto NameTestPredicate = [this, UserPreset, InNewText](TSharedPtr<IFilterPreset> TestPreset)
+			{
+				return TestPreset != UserPreset && TestPreset->GetDisplayText().ToString() == InNewText.ToString();
+			};
+
+			if (UserFilterPresets.ContainsByPredicate(NameTestPredicate) || SharedUserFilterPresets.ContainsByPredicate(NameTestPredicate))
+			{
+				OutErrorMessage = LOCTEXT("DuplicatePresetNames", "This name is already in use");
+				return false;
+			}
+			return true;
+		})
+		, FText::GetEmpty()
+	);
+	
+	const FText DisplayText = UserPreset->IsLocal() ? LOCTEXT("MakeSharedPresetLabel", "Make Shared Preset") : LOCTEXT("MakeSharedPresetLabel", "Make Local Preset");
+	const FText TooltipText = UserPreset->IsLocal() ? LOCTEXT("MakeSharedPresetToolTip", "Makes this preset a Shared User Preset (Config INI file has to be writable)") : LOCTEXT("MakeLocalPresetToolTip", "Makes this preset a Local Preset");
+
+	InSubMenuBuilder.AddMenuEntry(DisplayText,
+		TooltipText,
+		FSlateIcon(),
+		FUIAction(
+			FExecuteAction::CreateLambda([UserPreset, this]()
+			{
+				const bool bSuccess = UserPreset->IsLocal() ? UserPreset->MakeShared() : UserPreset->MakeLocal();
+				if (bSuccess)
+				{
+					RefreshFilterPresets();
+				}
+			}),
+			FCanExecuteAction::CreateLambda([]()
+			{
+				return FFilterPresetHelpers::CanModifySharedPreset();
+			})
+		)
+	);
+
+		InSubMenuBuilder.AddMenuEntry(LOCTEXT("PresetSaveLabel", "Save Preset"),
+		LOCTEXT("PresetSaveToolTip", "Saves the current filtering state as this User Preset"),
+		FSlateIcon(),
+		FUIAction(
+			FExecuteAction::CreateLambda([UserPreset, this]()
+			{
+				if (UserPreset.IsValid())
+				{
+					OnSavePreset.ExecuteIfBound(UserPreset);
+				}
+			})
+		)
+	);
+
+	InSubMenuBuilder.AddMenuEntry(LOCTEXT("PresetDeleteLabel", "Delete Preset"),
+		LOCTEXT("PresetDeleteToolTip", "Deletes the User Preset"),
+		FSlateIcon(),
+		FUIAction(
+			FExecuteAction::CreateLambda([UserPreset, this]()
+			{
+				if ( UserPreset->Delete())
+				{
+					RefreshFilterPresets();
+				}
+			})
+		)
+	);
+}
+
+void SFilterPresetList::TogglePreset(const TSharedPtr<IFilterPreset> InPreset)
+{
+	// Toggle off 
+	for (TSharedRef<SFilterPreset> Filter : Presets)
+	{
+		if (Filter->GetFilterPreset() == InPreset)
+		{
+			RemoveFilterPresetAndUpdate(Filter);
+			return;
+		}
+	}
+
+	// Toggle on
+	TSharedPtr<SFilterPreset> PresetWidget = AddFilterPreset(InPreset);
+	PresetWidget->SetEnabled(true);
+}
+
+bool SFilterPresetList::IsPresetEnabled(const TSharedPtr<IFilterPreset> InPreset) const 
+{
+	for (TSharedRef<SFilterPreset> Filter : Presets)
+	{
+		if (Filter->GetFilterPreset() == InPreset)
+		{
+			return true;
+		}
+	}
+
+	return false;
+}
+
+TSharedRef<SWidget> SFilterPresetList::MakeFilterPresetsMenu()
+{
+	FMenuBuilder MenuBuilder(/*bInShouldCloseWindowAfterMenuSelection=*/false, nullptr, nullptr, /*bCloseSelfOnly=*/true);
+
+	MenuBuilder.BeginSection("TracePresetsResetPresets");
+	{
+		MenuBuilder.AddMenuEntry(
+			LOCTEXT("FilterListResetPresets", "Reset Presets"),
+			LOCTEXT("FilterListResetToolTip", "Resets current presets selection"),
+			FSlateIcon(),
+			FUIAction(FExecuteAction::CreateSP(this, &SFilterPresetList::RemoveAllPresets))
+		);
+		
+		MenuBuilder.AddMenuEntry(
+			LOCTEXT("FilterListSavePresets", "Save as User Preset"),
+			LOCTEXT("FilterListSaveToolTip", "Saves the currently filtering state as a new User Preset"),
+			FSlateIcon(),
+			FUIAction(FExecuteAction::CreateLambda([this]() 
+			{
+				// Cache all user names to determine which preset is the newly created one
+				TArray<FString> Names;
+				Algo::Transform(UserFilterPresets, Names, [](TSharedPtr<IFilterPreset> Preset)
+				{
+					return Preset->GetName();
+				});
+
+				OnSavePreset.ExecuteIfBound(nullptr); 
+				RefreshFilterPresets();
+
+				for (const TSharedPtr<IFilterPreset>& Preset : UserFilterPresets)
+				{
+					if (!Names.Contains(Preset->GetName()))
+					{
+						// This is the newly created preset, enable it for UX
+						TSharedRef<SFilterPreset> Widget = AddFilterPreset(Preset);
+						EnableOnlyThisPreset(Widget);
+						break;
+					}
+				}
+			}))
+		);
+	}
+	MenuBuilder.EndSection(); // TracePresetsResetPresets
+
+	GenerateEnginePresetsMenu(MenuBuilder);
+	GenerateLocalUserPresetsMenu(MenuBuilder);
+	GenerateSharedUserPresetsMenu(MenuBuilder);
+
+	return MenuBuilder.MakeWidget();
+}
+
+TSharedRef<SFilterPreset> SFilterPresetList::AddFilterPreset(const TSharedPtr<IFilterPreset> FilterPreset)
+{
+	TSharedRef<SFilterPreset> NewFilter = SNew(SFilterPreset)
+		.FilterPreset(FilterPreset)
+		.OnPresetsChanged(OnPresetsChanged)
+		.OnRequestRemove(this, &SFilterPresetList::RemoveFilterPresetAndUpdate)
+		.OnRequestEnableAll(this, &SFilterPresetList::EnableAllPresets)
+		.OnRequestEnableOnly(this, &SFilterPresetList::EnableOnlyThisPreset)
+		.OnRequestDisableAll(this, &SFilterPresetList::DisableAllPresets)
+		.OnRequestRemoveAll(this, &SFilterPresetList::RemoveAllPresets)
+		.OnRequestDelete(this, & SFilterPresetList::DeletePreset)
+		.OnRequestSave(this, &SFilterPresetList::SavePreset)
+		.OnHighlightPreset(OnHighlightPreset);
+
+	Presets.Add(NewFilter);
+
+	// Add the widget as a child
+	FilterBox->AddSlot()
+	.Padding(3, 3)
+	[
+		NewFilter
+	];
+
+	return NewFilter;
+}
+
+void SFilterPresetList::EnableOnlyThisPreset(const TSharedRef<SFilterPreset>& PresetWidgetToEnable)
+{
+	for (TSharedRef<SFilterPreset>& PresetWidget : Presets )
+	{
+		const bool bEnable = (PresetWidget == PresetWidgetToEnable);
+		PresetWidget->SetEnabled(bEnable, false);
+	}
+
+	OnPresetsChanged.ExecuteIfBound();
+}
+
+void SFilterPresetList::RemoveFilterPresetAndUpdate(const TSharedRef<SFilterPreset>& PresetToRemove)
+{
+	FilterBox->RemoveSlot(PresetToRemove);
+	Presets.RemoveSingleSwap(PresetToRemove);
+
+	if (PresetToRemove->IsEnabled())
+	{
+		// Notify that a filter has changed
+		OnPresetsChanged.ExecuteIfBound();
+	}
+}
+
+void SFilterPresetList::DeletePreset(const TSharedRef<SFilterPreset>& PresetToDelete)
+{
+	if (PresetToDelete->GetFilterPreset()->Delete())
+	{
+		RefreshFilterPresets();
+	}
+}
+
+void SFilterPresetList::SavePreset(const TSharedRef<SFilterPreset>& PresetToSave)
+{
+	OnSavePreset.ExecuteIfBound(PresetToSave->GetFilterPreset());
+	EnableOnlyThisPreset(PresetToSave);
+}
+
+void SFilterPresetList::GetAllEnabledPresets(TArray<TSharedPtr<IFilterPreset>>& OutPresets) const
+{
+	for (const TSharedRef<SFilterPreset>& Filter : Presets)
+	{
+		if (Filter->IsEnabled())
+		{
+			OutPresets.Add(Filter->GetFilterPreset());
+		}
+	}
+}
+
+TSharedRef<SWidget> SFilterPresetList::ExternalMakeFilterPresetsMenu()
+{
+	return MakeFilterPresetsMenu();
+}
+
+#undef LOCTEXT_NAMESPACE
diff --git a/Engine/Plugins/Developer/TraceDataFiltering/Source/TraceDataFiltering/Private/SFilterPresetList.h b/Engine/Plugins/Developer/TraceDataFiltering/Source/TraceDataFiltering/Private/SFilterPresetList.h
new file mode 100644
--- /dev/null
+++ b/Engine/Plugins/Developer/TraceDataFiltering/Source/TraceDataFiltering/Private/SFilterPresetList.h
@@ -0,0 +1,125 @@
+// Copyright Epic Games, Inc. All Rights Reserved.
+
+#pragma once
+
+#include "CoreMinimal.h"
+#include "Input/Reply.h"
+#include "Widgets/DeclarativeSyntaxSupport.h"
+#include "Widgets/SWidget.h"
+#include "Widgets/SCompoundWidget.h"
+
+class FMenuBuilder;
+class SFilterPreset;
+class SWrapBox;
+
+/**
+ * A list of presets currently applied to insights filtering window
+ */
+class SFilterPresetList : public SCompoundWidget
+{
+public:
+	DECLARE_DELEGATE_OneParam(FOnSavePreset, const TSharedPtr<IFilterPreset>& /*PresetToSave, nullptr means create a new one */);
+	DECLARE_DELEGATE_OneParam(FOnHighlightPreset, const TSharedPtr<IFilterPreset>& /*PresetToHighlight, nullptr means stop highlighting */);
+
+	SLATE_BEGIN_ARGS( SFilterPresetList ){}
+		/** Delegate for when the selected presets have changed */
+		SLATE_EVENT(FOnPresetsChanged, OnPresetsChanged)
+		/** Delegate for when a preset is to be saved, given the current filtering state */
+		SLATE_EVENT(FOnSavePreset, OnSavePreset)
+		/** Delegate for when a specific preset is to be highlighted outside of this widget */
+		SLATE_EVENT(FOnHighlightPreset, OnHighlightPreset)
+	SLATE_END_ARGS()
+
+	/** Constructs this widget with InArgs */		
+	void Construct( const FArguments& InArgs );
+
+	virtual ~SFilterPresetList();
+
+	/** Begin SCompoundWidget overrides */
+	virtual FReply OnMouseButtonUp(const FGeometry& MyGeometry, const FPointerEvent& MouseEvent);
+	/** End SCompoundWidget overrides */
+
+	/** Returns all filter presets currently set as enabled within this widget */
+	void GetAllEnabledPresets(TArray<TSharedPtr<IFilterPreset>>& OutPresets) const;
+
+	/** Returns the context menu widget generated by MakeFilterPresetsMenu */
+	TSharedRef<SWidget> ExternalMakeFilterPresetsMenu();	
+
+private:
+	/** Refreshes the filter preset objects and child widgets */
+	void RefreshFilterPresets();
+
+	/** Returns true if any Presets are applied */
+	bool HasAnyPresets() const;
+
+	/** Enables all Presets */
+	void EnableAllPresets();
+
+	/** Disables any applied Presets */
+	void DisableAllPresets();
+
+	/** Removes all Presets in the list */
+	void RemoveAllPresets();
+
+	/** Adds a child widget for the specified filter preset object */
+	TSharedRef<SFilterPreset> AddFilterPreset(const TSharedPtr<IFilterPreset> FilterPreset);
+
+	/** Removes the child widget, and calls OnPresetsChanged to ensure the state is updated outside of this widget */
+	void RemoveFilterPresetAndUpdate(const TSharedRef<SFilterPreset>& PresetToRemove);
+
+	/** Handler for when the enable only this button was clicked on a single preset */
+	void EnableOnlyThisPreset(const TSharedRef<SFilterPreset>& PresetWidgetToEnable);
+
+	void TogglePreset(const TSharedPtr<IFilterPreset> InPreset);
+	bool IsPresetEnabled(const TSharedPtr<IFilterPreset> InPreset) const;
+
+	/** Saves any settings to config that should be persistent between editor sessions */
+	void SaveSettings(const FString& IniFilename) const;
+
+	/** Loads any settings to config that should be persistent between editor sessions */
+	void LoadSettings(const FString& IniFilename);
+
+	TSharedRef<SWidget> MakeFilterPresetsMenu();
+	void GenerateEnginePresetsMenu(FMenuBuilder& InMenuBuilder);
+	void GenerateLocalUserPresetsMenu(FMenuBuilder& InMenuBuilder);
+	void GenerateCommonPresetEntries(FMenuBuilder& InSubMenuBuilder, const TSharedPtr<IFilterPreset>& UserPreset);
+	void GenerateSharedUserPresetsMenu(FMenuBuilder& InMenuBuilder);
+
+	/** Delete the preset object contained by the widget */
+	void DeletePreset(const TSharedRef<SFilterPreset>& PresetToDelete);
+	/** Save the preset object contained by the widget */
+	void SavePreset(const TSharedRef<SFilterPreset>& PresetToSave);	
+private:
+	/** The horizontal box which contains all the Presets */
+	TSharedPtr<SWrapBox> FilterBox;
+
+	/** All SFilterPreset in the list */
+	TArray<TSharedRef<SFilterPreset>> Presets;
+	
+	/** All Filter Preset objects */
+	TArray<TSharedPtr<IFilterPreset>> AllFilterPresets;
+
+	/** All User Filter Preset objects */
+	TArray<TSharedPtr<IFilterPreset>> UserFilterPresets;
+	/** All Shared User Filter Preset objects */
+	TArray<TSharedPtr<IFilterPreset>> SharedUserFilterPresets;
+	/** All Engine Preset objects */
+	TArray<TSharedPtr<IFilterPreset>> EngineFilterPresets;
+
+	/** Weak pointers to filter preset containers */
+	TWeakObjectPtr<ULocalFilterPresetContainer> PresetContainer;
+	TWeakObjectPtr<USharedFilterPresetContainer> SharedPresetContainer;
+
+	/** Delegate for when Presets have changed */
+	FOnPresetsChanged OnPresetsChanged;
+
+	/** Delegate for when a Preset is to be saved */
+	FOnSavePreset OnSavePreset;
+
+	/** Delegate for when a Preset is to be highlighted */
+	FOnHighlightPreset OnHighlightPreset;
+
+	/** Cached arrays of currently active and enabled filter names, used for saving/restoring state during refreshes */
+	TArray<FString> ActiveFilterNames;
+	TArray<FString> EnabledFilterNames;
+};
diff --git a/Engine/Plugins/Developer/TraceDataFiltering/Source/TraceDataFiltering/Private/STraceDataFilterWidget.cpp b/Engine/Plugins/Developer/TraceDataFiltering/Source/TraceDataFiltering/Private/STraceDataFilterWidget.cpp
new file mode 100644
--- /dev/null
+++ b/Engine/Plugins/Developer/TraceDataFiltering/Source/TraceDataFiltering/Private/STraceDataFilterWidget.cpp
@@ -0,0 +1,691 @@
+// Copyright Epic Games, Inc. All Rights Reserved.
+
+#include "STraceDataFilterWidget.h"
+#include "EditorStyleSet.h"
+#include "Widgets/Layout/SScrollBorder.h"
+#include "Widgets/Text/STextBlock.h"
+#include "Widgets/Input/SComboButton.h"
+#include "Widgets/Input/SSearchBox.h"
+#include "Templates/SharedPointer.h"
+#include "SFilterPresetList.h"
+#include "STraceObjectRowWidget.h"
+
+#include "ITraceObject.h"
+#include "TraceChannel.h"
+
+#define LOCTEXT_NAMESPACE "STraceDataFilterWidget"
+
+TSet<FString> STraceDataFilterWidget::LastExpandedObjectNames;
+
+STraceDataFilterWidget::STraceDataFilterWidget() : FilterService(FEventFilterService::Get()), bNeedsTreeRefresh(false), bHighlightingPreset(false)
+{
+}
+
+STraceDataFilterWidget::~STraceDataFilterWidget()
+{
+	// Save expansion and move to static array
+	SaveItemsExpansion();
+	LastExpandedObjectNames = ExpandedObjectNames;
+	ExpandedObjectNames.Empty();
+}
+
+void STraceDataFilterWidget::Construct(const FArguments& InArgs)
+{
+	ConstructSearchBoxFilter();
+	ConstructFilterHandler();
+	ConstructTreeview();
+
+	// Retrieve trace services module for session selection widget
+	ITraceServicesModule& TraceServicesModule = FModuleManager::LoadModuleChecked<ITraceServicesModule>("TraceServices");
+
+	ChildSlot
+	[		
+		SNew(SBorder)
+		.Padding(4)
+		.BorderImage(FEditorStyle::GetBrush("ToolPanel.GroupBorder"))
+		[
+			SNew( SVerticalBox )
+
+			+ SVerticalBox::Slot()
+			.AutoHeight()
+			.Padding(FMargin(0.0f, 0.0f, 0.0f, 2.0f))
+			[
+				// Filtering button and search box widgets
+				SAssignNew(OptionsWidget, SHorizontalBox)
+				+SHorizontalBox::Slot()
+				.AutoWidth()
+				.VAlign(VAlign_Center)
+				.Padding(0.0f, 0.0f, 2.0f, 0.0f)
+				[
+					SNew(SComboButton)
+					.Visibility(EVisibility::Visible)
+					.ComboButtonStyle(FEditorStyle::Get(), "GenericFilters.ComboButtonStyle")
+					.ForegroundColor(FLinearColor::White)
+					.ContentPadding(0.0f)
+					.OnGetMenuContent(this, &STraceDataFilterWidget::MakeAddFilterMenu)
+					.ButtonContent()
+					[
+						SNew(SHorizontalBox)
+						+SHorizontalBox::Slot()
+						.AutoWidth()
+						.VAlign(VAlign_Center)
+						[
+							SNew(STextBlock)
+							.TextStyle(FEditorStyle::Get(), "GenericFilters.TextStyle")
+							.Font(FEditorStyle::Get().GetFontStyle("FontAwesome.9"))
+							.Text(FText::FromString(FString(TEXT("\xf0b0"))) /*fa-filter*/)
+						]
+						+SHorizontalBox::Slot()
+						.AutoWidth()
+						.Padding(2, 0, 0, 0)
+						.VAlign(VAlign_Center)
+						[
+							SNew( STextBlock )
+							.TextStyle(FEditorStyle::Get(), "GenericFilters.TextStyle")
+							.Text( LOCTEXT("PresetsMenuLabel", "Filter Presets") )
+						]
+					]
+				]
+				+SHorizontalBox::Slot()
+				.FillWidth(1.0f)
+				[
+					SAssignNew(SearchBoxWidget, SSearchBox)
+					.SelectAllTextWhenFocused( true )
+					.HintText( LOCTEXT( "SearchBoxHint", "Search Trace Events..."))
+					.OnTextChanged(this, &STraceDataFilterWidget::OnSearchboxTextChanged)
+				]
+			]
+
+			+ SVerticalBox::Slot()
+			.Padding(0.0f, 0.0f, 0.0f, 2.0f)
+			.AutoHeight()
+			[
+				// Presets bar widget
+				SNew(SHorizontalBox)
+				+SHorizontalBox::Slot()
+				.FillWidth(1.0f)
+				[
+					SAssignNew(FilterPresetsListWidget, SFilterPresetList)
+					.OnPresetsChanged(this, &STraceDataFilterWidget::OnPresetsChanged)
+					.OnSavePreset(this, &STraceDataFilterWidget::OnSavePreset)
+					.OnHighlightPreset(this, &STraceDataFilterWidget::OnHighlightPreset)
+				]
+			]
+
+			+ SVerticalBox::Slot()	
+			.Padding(0.0f, 0.0f, 0.0f, 2.0f)
+			[
+				SNew(SOverlay)
+				+ SOverlay::Slot()
+				.HAlign(HAlign_Fill)
+				.VAlign(VAlign_Fill)
+				[
+					SNew(SScrollBorder, Treeview.ToSharedRef())
+					[
+						Treeview.ToSharedRef()					
+					]
+				]
+			]			
+		]	
+	];
+
+	/** Setup attribute to enable/disabled all widgets according to whether or not there is a valid session to represent */
+	TAttribute<bool> EnabledAttribute;
+	EnabledAttribute.Bind(this, &STraceDataFilterWidget::HasValidFilterSession);
+
+	Treeview->SetEnabled(EnabledAttribute);
+	OptionsWidget->SetEnabled(EnabledAttribute);
+	FilterPresetsListWidget->SetEnabled(EnabledAttribute);
+
+	/** Restore expansion state when recreating this window */
+	ExpandedObjectNames = LastExpandedObjectNames;
+	RestoreItemsExpansion();
+}
+
+void STraceDataFilterWidget::OnSavePreset(const TSharedPtr<IFilterPreset>& Preset)
+{
+	if (SessionFilterService.IsValid())
+	{
+		/** Save to preset if one was provided, otherwise create a new one */
+		if (Preset.IsValid())
+		{
+			Preset->Save(RootItems);
+		}
+		else
+		{
+			FFilterPresetHelpers::CreateNewPreset(RootItems);
+		}
+	}
+}
+
+void STraceDataFilterWidget::OnPresetsChanged()
+{
+	TArray<TSharedPtr<IFilterPreset>> Presets;
+	FilterPresetsListWidget->GetAllEnabledPresets(Presets);
+
+	if (SessionFilterService.IsValid())
+	{
+		SessionFilterService->UpdateFilterPresets(Presets);
+	}
+}
+
+void STraceDataFilterWidget::OnHighlightPreset(const TSharedPtr<IFilterPreset>& Preset)
+{
+	if (Treeview.IsValid())
+	{
+		/** Update treeview so that any whitelisted entry (as part of Preset) is highlighted and expanded */
+		Treeview->ClearHighlightedItems();
+		if (Preset.IsValid())
+		{
+			if (!bHighlightingPreset)
+			{
+				/** Store current expansion, so we can reset once highlighting has finished */
+				SaveItemsExpansion();
+				bHighlightingPreset = true;
+			}
+
+			TArray<FString> Names;
+			Preset->GetWhitelistedNames(Names);
+
+			EnumerateAllItems([this, Names](TSharedPtr<ITraceObject> Object) -> void
+			{
+				if (Names.Contains(Object->GetName()))
+				{
+					Treeview->SetItemHighlighted(Object, true);
+					// TODO need to set parent expansion as well
+					Treeview->SetItemExpansion(Object, true);
+
+					SetParentExpansionRecursively(Object, true);
+				}
+			});
+		}
+		else
+		{
+			bHighlightingPreset = false;
+			RestoreItemsExpansion();
+		}
+	}
+}
+
+void STraceDataFilterWidget::OnSearchboxTextChanged(const FText& FilterText)
+{
+	bNeedsTreeRefresh = true;
+	const bool bFilterSet = !FilterText.IsEmpty();
+	if (bFilterSet != TreeviewFilterHandler->GetIsEnabled())
+	{
+		TreeviewFilterHandler->SetIsEnabled(bFilterSet);
+
+		if (bFilterSet)
+		{
+			SaveItemsExpansion();
+		}
+		else
+		{
+			RestoreItemsExpansion();
+		}
+	}
+
+	SearchBoxWidgetFilter->SetRawFilterText(FilterText);
+	SearchBoxWidget->SetError(SearchBoxWidgetFilter->GetFilterErrorText());
+}
+
+void STraceDataFilterWidget::ConstructTreeview()
+{
+	SAssignNew(Treeview, STreeView<TSharedPtr<ITraceObject>>)
+	.ItemHeight(20.0f)
+	.OnGetChildren(TreeviewFilterHandler.ToSharedRef(), &TreeFilterHandler<TSharedPtr<ITraceObject>>::OnGetFilteredChildren)
+	.OnGenerateRow(this, &STraceDataFilterWidget::OnGenerateRow)
+	.OnContextMenuOpening(this, &STraceDataFilterWidget::OnContextMenuOpening)
+	.OnMouseButtonDoubleClick(this, &STraceDataFilterWidget::OnItemDoubleClicked)
+	.OnSetExpansionRecursive_Lambda([this](TSharedPtr<ITraceObject> InObject, bool bInExpansionState)
+	{
+		SetExpansionRecursively(InObject, bInExpansionState);
+	})
+	.TreeItemsSource(&TreeItems);
+	
+	TreeviewFilterHandler->SetTreeView(Treeview.Get());
+}
+
+void STraceDataFilterWidget::OnItemDoubleClicked(TSharedPtr<ITraceObject> InObject) const
+{
+	TSharedPtr<ITableRow> Row = Treeview->WidgetFromItem(InObject);
+	if (Row.IsValid())
+	{
+		SetExpansionRecursively(InObject, !Row->IsItemExpanded());
+	}
+}
+
+void STraceDataFilterWidget::ConstructSearchBoxFilter()
+{
+	SearchBoxWidgetFilter = MakeShareable(new TTextFilter<TSharedPtr<ITraceObject>>(TTextFilter<TSharedPtr<ITraceObject>>::FItemToStringArray::CreateLambda([](TSharedPtr<ITraceObject> Object, TArray<FString>& OutStrings)
+	{
+		Object->GetSearchString(OutStrings);
+	})));
+}
+
+void STraceDataFilterWidget::ConstructFilterHandler()
+{
+	TreeviewFilterHandler = MakeShareable(new TreeFilterHandler<TSharedPtr<ITraceObject>>());
+	TreeviewFilterHandler->SetFilter(SearchBoxWidgetFilter.Get());
+	TreeviewFilterHandler->SetRootItems(&RootItems, &TreeItems);
+	TreeviewFilterHandler->SetGetChildrenDelegate(TreeFilterHandler<TSharedPtr<ITraceObject>>::FOnGetChildren::CreateLambda([](TSharedPtr<ITraceObject> InParent, TArray<TSharedPtr<ITraceObject>>& OutChildren)
+	{
+		InParent->GetChildren(OutChildren);
+	}));
+}
+
+TSharedRef<ITableRow> STraceDataFilterWidget::OnGenerateRow(TSharedPtr<ITraceObject> InItem, const TSharedRef<STableViewBase>& OwnerTable)
+{
+	return SNew(STraceObjectRowWidget, OwnerTable, InItem)
+		.HighlightText(SearchBoxWidgetFilter.Get(), &TTextFilter<TSharedPtr<ITraceObject>>::GetRawFilterText);
+}
+
+TSharedRef<SWidget> STraceDataFilterWidget::MakeAddFilterMenu()
+{
+	return FilterPresetsListWidget->ExternalMakeFilterPresetsMenu();
+}
+
+TSharedPtr<SWidget> STraceDataFilterWidget::OnContextMenuOpening() const
+{
+	FMenuBuilder MenuBuilder(/*bInShouldCloseWindowAfterMenuSelection=*/true, nullptr, nullptr, /*bCloseSelfOnly=*/true);
+
+	static const FName FilteringSectionHook("FilteringState");
+	MenuBuilder.BeginSection(FilteringSectionHook, LOCTEXT("FilteringSectionLabel", "Filtering"));
+	{
+		MenuBuilder.AddMenuEntry(LOCTEXT("EnableAllRowsLabel", "Enable All"), LOCTEXT("EnableAllRowsTooltip", "Sets entire hierarchy to be non-filtered."), FSlateIcon(), 
+			FUIAction(FExecuteAction::CreateLambda([this]()
+				{
+					EnumerateAllItems([this](TSharedPtr<ITraceObject> Object) -> void
+					{
+						Object->SetIsFiltered(false);
+					});
+				}),
+				FCanExecuteAction(),
+				FGetActionCheckState(),
+				FIsActionButtonVisible::CreateLambda([this]() -> bool
+				{
+					return Treeview->GetNumItemsSelected() == 0 && EnumerateAllItems(TFunction<bool(TSharedPtr<ITraceObject> InItem)>([this](TSharedPtr<ITraceObject> Object) -> bool
+					{
+						return Object->IsFiltered();
+					}));
+				})
+			)
+		);
+
+		MenuBuilder.AddMenuEntry(LOCTEXT("DisableAllRowsLabel", "Disable All"), LOCTEXT("DisableAllRowsTooltip", "Sets entire hierarchy to be filtered."), FSlateIcon(), 
+			FUIAction(
+				FExecuteAction::CreateLambda([this]()
+				{
+					EnumerateAllItems([this](TSharedPtr<ITraceObject> Object) -> void
+					{
+						Object->SetIsFiltered(true);
+					});
+				}),
+				FCanExecuteAction(),
+				FGetActionCheckState(),
+				FIsActionButtonVisible::CreateLambda([this]() -> bool
+				{
+					return Treeview->GetNumItemsSelected() == 0 && EnumerateAllItems(TFunction<bool(TSharedPtr<ITraceObject> InItem)>([this](TSharedPtr<ITraceObject> Object) -> bool
+					{
+						return !Object->IsFiltered();
+					}));
+				})
+
+			)
+		);
+
+		MenuBuilder.AddMenuEntry(LOCTEXT("EnableRowsLabel", "Enable Selected"), LOCTEXT("EnableRowsTooltip", "Sets the selected Node(s) to be non-filtered."), FSlateIcon(),
+			FUIAction(
+				FExecuteAction::CreateRaw(this, &STraceDataFilterWidget::EnumerateSelectedItems, TFunction<void(TSharedPtr<ITraceObject> InItem)>([this](TSharedPtr<ITraceObject> Object) -> void
+				{
+					Object->SetIsFiltered(false);
+
+					SetExpansionRecursively(Object, true);
+				})),
+				FCanExecuteAction(),
+				FGetActionCheckState(),
+				FIsActionButtonVisible::CreateLambda([this]() -> bool
+				{
+					return EnumerateSelectedItems(TFunction<bool(TSharedPtr<ITraceObject> InItem)>([this](TSharedPtr<ITraceObject> Object) -> bool
+					{
+						return Object->IsFiltered();
+					}));
+				})
+			)
+		);
+
+		MenuBuilder.AddMenuEntry(LOCTEXT("DisableRowsLabel", "Disable Selected"), LOCTEXT("DisableRowsTooltip", "Sets the selected Node(s) to be filtered."), FSlateIcon(),
+			FUIAction(
+				FExecuteAction::CreateRaw(this, &STraceDataFilterWidget::EnumerateSelectedItems, TFunction<void(TSharedPtr<ITraceObject> InItem)>([this](TSharedPtr<ITraceObject> Object) -> void
+				{
+					Object->SetIsFiltered(true);
+				})),
+				FCanExecuteAction(),
+				FGetActionCheckState(),
+				FIsActionButtonVisible::CreateLambda([this]() -> bool
+				{
+					return EnumerateSelectedItems(TFunction<bool(TSharedPtr<ITraceObject> InItem)>([this](TSharedPtr<ITraceObject> Object) -> bool
+					{
+						return !Object->IsFiltered();
+					}));
+				})
+			)
+		);
+	}
+	MenuBuilder.EndSection();
+
+	static const FName ExpansionSectionHook("ExpansionState");
+	MenuBuilder.BeginSection(ExpansionSectionHook, LOCTEXT("ExpansionSectionLabel", "Expansion"));
+	{
+		MenuBuilder.AddMenuEntry(LOCTEXT("ExpandAllRowsLabel", "Expand All"), LOCTEXT("ExpandAllRowsTooltip", "Expands the entire hierarchy."), FSlateIcon(), FUIAction(
+				FExecuteAction::CreateLambda([this]()
+				{
+					for (const TSharedPtr<ITraceObject>& Object : TreeItems)
+					{
+						SetExpansionRecursively(Object, true);
+					}
+				}),
+				FCanExecuteAction(),
+				FGetActionCheckState(),
+				FIsActionButtonVisible::CreateLambda([this]() -> bool
+				{
+					return Treeview->GetNumItemsSelected() == 0 && EnumerateAllItems(TFunction<bool(TSharedPtr<ITraceObject> InItem)>([this](TSharedPtr<ITraceObject> Object) -> bool
+					{
+						return !Treeview->IsItemExpanded(Object);
+					}));
+				})
+			)
+		);
+		
+		MenuBuilder.AddMenuEntry(LOCTEXT("CollapseAllRowsLabel", "Collapse All"), LOCTEXT("CollapseAllRowsTooltip", "Collapses the entire hierarchy."), FSlateIcon(), 
+			FUIAction(
+				FExecuteAction::CreateLambda([this]()
+				{
+					for (const TSharedPtr<ITraceObject>& Object : TreeItems)
+					{
+						SetExpansionRecursively(Object, false);
+					}
+				}),
+				FCanExecuteAction(),
+				FGetActionCheckState(),
+				FIsActionButtonVisible::CreateLambda([this]() -> bool
+				{
+					return Treeview->GetNumItemsSelected() == 0 && EnumerateAllItems(TFunction<bool(TSharedPtr<ITraceObject> InItem)>([this](TSharedPtr<ITraceObject> Object) -> bool
+					{
+						return Treeview->IsItemExpanded(Object);
+					}));
+				})
+			)
+		);
+
+		MenuBuilder.AddMenuEntry(LOCTEXT("ExpandRowsLabel", "Expand Selected"), LOCTEXT("ExpandRowsTooltip", "Expands the selected Node(s)."), FSlateIcon(), 
+			FUIAction(
+				FExecuteAction::CreateRaw(this, &STraceDataFilterWidget::EnumerateSelectedItems, TFunction<void(TSharedPtr<ITraceObject> InItem)>([this](TSharedPtr<ITraceObject> Object) -> void
+				{
+					SetExpansionRecursively(Object, true);
+				})),
+				FCanExecuteAction(),
+				FGetActionCheckState(),
+				FIsActionButtonVisible::CreateLambda([this]() -> bool
+				{
+					return EnumerateSelectedItems(TFunction<bool(TSharedPtr<ITraceObject> InItem)>([this](TSharedPtr<ITraceObject> Object) -> bool
+					{
+						return !Treeview->IsItemExpanded(Object);
+					}));
+				})
+			)
+		);
+
+		MenuBuilder.AddMenuEntry(LOCTEXT("CollapseRowsLabel", "Collapse Selected"), LOCTEXT("CollapseRowsTooltip", "Collapse the selected Node(s)."), FSlateIcon(), 
+			FUIAction(
+				FExecuteAction::CreateRaw(this, &STraceDataFilterWidget::EnumerateSelectedItems, TFunction<void(TSharedPtr<ITraceObject> InItem)>([this](TSharedPtr<ITraceObject> Object) -> void
+				{
+					SetExpansionRecursively(Object, false);
+				})),
+				FCanExecuteAction(),
+				FGetActionCheckState(),
+				FIsActionButtonVisible::CreateLambda([this]() -> bool
+				{
+					return EnumerateSelectedItems(TFunction<bool(TSharedPtr<ITraceObject> InItem)>([this](TSharedPtr<ITraceObject> Object) -> bool
+					{
+						return Treeview->IsItemExpanded(Object);
+					}));
+				})
+			)
+		);
+	}  	
+	MenuBuilder.EndSection();
+
+	return MenuBuilder.MakeWidget();
+}
+
+void STraceDataFilterWidget::SaveItemsExpansion()
+{
+	ExpandedObjectNames.Empty();
+
+	if (Treeview.IsValid())
+	{
+		TSet<TSharedPtr<ITraceObject>> ExpandedItems;
+		Treeview->GetExpandedItems(ExpandedItems);
+
+		for (TSharedPtr<ITraceObject> Item : ExpandedItems)
+		{
+			if (Item.IsValid())
+			{
+				ExpandedObjectNames.Add(Item->GetName());
+			}
+		}
+	}
+}
+
+void STraceDataFilterWidget::RestoreItemsExpansion()
+{
+	EnumerateAllItems([this](TSharedPtr<ITraceObject> Object) -> void
+	{
+		const bool bExpanded = ExpandedObjectNames.Contains(Object->GetName());
+		SetExpansionRecursively(Object, bExpanded);
+	});
+
+	ExpandedObjectNames.Empty();
+}
+
+void STraceDataFilterWidget::SaveItemSelection()
+{
+	SelectedObjectNames.Empty();
+
+	if (Treeview.IsValid())
+	{
+		EnumerateSelectedItems([this](TSharedPtr<ITraceObject> InObject)
+		{
+			if (InObject.IsValid())
+			{
+				SelectedObjectNames.Add(InObject->GetName());
+			}
+		});
+	}
+}
+
+void STraceDataFilterWidget::RestoreItemSelection()
+{
+	if (Treeview.IsValid())
+	{
+		TArray<TSharedPtr<ITraceObject>> SelectedItems;
+		EnumerateAllItems([this, &SelectedItems](TSharedPtr<ITraceObject> Object) -> void
+		{
+			if (SelectedObjectNames.Contains(Object->GetName()))
+			{
+				SelectedItems.Add(Object);
+			}
+		});
+
+		Treeview->SetItemSelection(SelectedItems, true);
+	}
+	
+	SelectedObjectNames.Empty();
+}
+
+void STraceDataFilterWidget::SetCurrentAnalysisSession(Trace::FSessionHandle Handle)
+{
+	SessionFilterService = FilterService.GetFilterServiceByHandle(Handle);
+	/** Refresh presets so config loaded state is directly applied */
+	OnPresetsChanged();
+	/** Refresh data driving the treeview */
+	RefreshTreeviewData();
+}
+
+TSharedRef<ITraceObject> STraceDataFilterWidget::AddFilterableObject(const FTraceObjectInfo& Event, FString ParentName)
+{
+	// Retrieve any child objects, and recursively add those first 
+	TArray<FTraceObjectInfo> Events;
+	SessionFilterService->GetChildObjects(Event.Hash, Events);
+
+	TArray<TSharedPtr<ITraceObject>> ChildPtrs;	
+	for (const FTraceObjectInfo& ChildEvent : Events)
+	{
+		TSharedRef<ITraceObject> EventItem = AddFilterableObject(ChildEvent, Event.Name);
+		ChildPtrs.Add(EventItem);
+	}
+
+	TSharedRef<FTraceChannel> SharedItem = MakeShareable(new FTraceChannel(Event.Name, ParentName, Event.Hash, !Event.bEnabled, ChildPtrs, SessionFilterService));
+	ParentToChild.Add(SharedItem, ChildPtrs);
+
+	for (TSharedPtr<ITraceObject> ChildObject : ChildPtrs)
+	{
+		ChildToParent.Add(ChildObject, SharedItem);
+	}
+
+	FlatItems.Add(SharedItem);
+
+	return SharedItem;
+}
+
+bool STraceDataFilterWidget::HasValidFilterSession() const
+{
+	return SessionFilterService.IsValid();
+}
+
+void STraceDataFilterWidget::SetExpansionRecursively(const TSharedPtr<ITraceObject>& InObject, bool bShouldExpandItem) const
+{
+	Treeview->SetItemExpansion(InObject, bShouldExpandItem);
+
+	TArray<TSharedPtr<ITraceObject>> Children;
+	InObject->GetChildren(Children);
+
+	for (TSharedPtr<ITraceObject>& ChildObject : Children)
+	{
+		SetExpansionRecursively(ChildObject, bShouldExpandItem);
+	}
+}
+
+void STraceDataFilterWidget::SetParentExpansionRecursively(const TSharedPtr<ITraceObject>& InObject, bool bShouldExpandItem) const
+{
+	if (const TSharedPtr<ITraceObject>* ParentObject = ChildToParent.Find(InObject))
+	{
+		Treeview->SetItemExpansion(*ParentObject, bShouldExpandItem);
+		SetParentExpansionRecursively(*ParentObject, bShouldExpandItem);
+	}
+}
+
+void STraceDataFilterWidget::EnumerateSelectedItems(TFunction<void(TSharedPtr<ITraceObject> InItem)> InFunction) const
+{
+	TArray<TSharedPtr<ITraceObject>> SelectedObjects;
+	Treeview->GetSelectedItems(SelectedObjects);
+
+	for (const TSharedPtr<ITraceObject>& Object : SelectedObjects)
+	{
+		InFunction(Object);
+	}
+}
+
+bool STraceDataFilterWidget::EnumerateSelectedItems(TFunction<bool(TSharedPtr<ITraceObject> InItem)> InFunction) const
+{
+	TArray<TSharedPtr<ITraceObject>> SelectedObjects;
+	Treeview->GetSelectedItems(SelectedObjects);
+
+
+	bool bState = false;
+	for (const TSharedPtr<ITraceObject>& Object : SelectedObjects)
+	{
+		bState |= InFunction(Object);
+	}
+
+	return bState;
+}
+
+void STraceDataFilterWidget::EnumerateAllItems(TFunction<void(TSharedPtr<ITraceObject> InItem)> InFunction) const
+{
+	for (const TSharedPtr<ITraceObject>& Object : FlatItems)
+	{
+		InFunction(Object);
+	}
+}
+
+bool STraceDataFilterWidget::EnumerateAllItems(TFunction<bool(TSharedPtr<ITraceObject> InItem)> InFunction) const
+{
+	bool bState = false;
+	for (const TSharedPtr<ITraceObject>& Object : FlatItems)
+	{
+		bState |= InFunction(Object);
+	}
+	return bState;
+}
+
+void STraceDataFilterWidget::RefreshTreeviewData()
+{
+	if (SessionFilterService.IsValid())
+	{
+		SyncTimeStamp = SessionFilterService->GetTimestamp();
+
+		/** Save expansion and selection */
+		SaveItemsExpansion();
+		SaveItemSelection();
+
+		TArray<FTraceObjectInfo> RootEvents;
+		SessionFilterService->GetRootObjects(RootEvents);
+		
+		ParentToChild.Empty();
+		ChildToParent.Empty();
+		RootItems.Empty();		
+
+		for (const FTraceObjectInfo& RootEvent : RootEvents)
+		{
+			TSharedRef<ITraceObject> TraceObject = AddFilterableObject(RootEvent, FString());
+			RootItems.Add(TraceObject);
+		}
+
+		TreeviewFilterHandler->RefreshAndFilterTree();
+
+		RestoreItemsExpansion();
+		RestoreItemSelection();
+	}
+}
+
+void STraceDataFilterWidget::Tick(const FGeometry& AllottedGeometry, const double InCurrentTime, const float InDeltaTime)
+{
+	if (SessionFilterService.IsValid() )
+	{
+		if (SessionFilterService->GetTimestamp() != SyncTimeStamp)
+		{
+			RefreshTreeviewData();
+		}
+	}
+	else
+	{
+		ITraceServicesModule& TraceServicesModule = FModuleManager::LoadModuleChecked<ITraceServicesModule>("TraceServices");
+		TArray<Trace::FSessionHandle> LiveSessionHandles;
+		TraceServicesModule.GetSessionService()->GetLiveSessions(LiveSessionHandles);
+		for (const Trace::FSessionHandle& Handle : LiveSessionHandles)
+		{
+			SetCurrentAnalysisSession(Handle);
+			break;
+		}
+	}
+
+	if (bNeedsTreeRefresh)
+	{
+		TreeviewFilterHandler->RefreshAndFilterTree();
+	}
+}
+
+
+#undef LOCTEXT_NAMESPACE // "STraceDataFilterWidget"
\ No newline at end of file
diff --git a/Engine/Plugins/Developer/TraceDataFiltering/Source/TraceDataFiltering/Private/STraceDataFilterWidget.h b/Engine/Plugins/Developer/TraceDataFiltering/Source/TraceDataFiltering/Private/STraceDataFilterWidget.h
new file mode 100644
--- /dev/null
+++ b/Engine/Plugins/Developer/TraceDataFiltering/Source/TraceDataFiltering/Private/STraceDataFilterWidget.h
@@ -0,0 +1,126 @@
+// Copyright Epic Games, Inc. All Rights Reserved.
+
+#pragma once
+
+#include "Templates/SharedPointer.h"
+#include "ITraceObject.h"
+#include "Misc/IFilter.h"
+#include "Framework/Views/TreeFilterHandler.h"
+#include "Misc/TextFilter.h"
+
+class ITableRow;
+class SFilterPresetList;
+class SSessionSelector;
+
+class STraceDataFilterWidget : public SCompoundWidget
+{
+public:
+	/** Default constructor. */
+	STraceDataFilterWidget();
+
+	/** Virtual destructor. */
+	virtual ~STraceDataFilterWidget();
+
+	SLATE_BEGIN_ARGS(STraceDataFilterWidget) {}
+	SLATE_END_ARGS()
+
+	/** Constructs this widget. */
+	void Construct(const FArguments& InArgs);
+
+	virtual void Tick(const FGeometry& AllottedGeometry, const double InCurrentTime, const float InDeltaTime) override;
+protected:
+	/** Callback from SFilterPresetList, should save the current tree view filter state as the specified preset */
+	void OnSavePreset(const TSharedPtr<IFilterPreset>& Preset);
+	/** Callback from SFilterPresetList, should update filtering state according to currently active Filter Presets */
+	void OnPresetsChanged();
+	/** Callback from SFilterPresetList, should highlight the tree view items encompassed by the specified preset */
+	void OnHighlightPreset(const TSharedPtr<IFilterPreset>& Preset);
+
+	void OnSearchboxTextChanged(const FText& FilterText);	
+	void OnItemDoubleClicked(TSharedPtr<ITraceObject> InObject) const;
+
+	/** Constructing various widgets */
+	void ConstructTreeview();
+	void ConstructSearchBoxFilter();
+	void ConstructFilterHandler();
+
+	/** Callback for whenever a different analysis session (by handle), has been selected through the SessionSelector widget */
+	void SetCurrentAnalysisSession(Trace::FSessionHandle Handle);
+	bool HasValidFilterSession() const;
+		
+	void RefreshTreeviewData();
+		
+	/** Sets the treeview item expansion state recursively for the specified object */
+	void SetExpansionRecursively(const TSharedPtr<ITraceObject>& InObject, bool bShouldExpandItem) const;
+	void SetParentExpansionRecursively(const TSharedPtr<ITraceObject>& InObject, bool bShouldExpandItem) const;	
+
+	/** Enumerates the (currently selected) treeview items and calls InFunction for each entry */
+	void EnumerateSelectedItems(TFunction<void(TSharedPtr<ITraceObject> InItem)> InFunction) const;
+	void EnumerateAllItems(TFunction<void(TSharedPtr<ITraceObject> InItem)> InFunction) const;
+	bool EnumerateAllItems(TFunction<bool(TSharedPtr<ITraceObject> InItem)> InFunction) const;
+	bool EnumerateSelectedItems(TFunction<bool(TSharedPtr<ITraceObject> InItem)> InFunction) const;
+
+	/** Creates and adds a new treeview object entry, given the specified data */
+	TSharedRef<ITraceObject> AddFilterableObject(const FTraceObjectInfo& Event, FString ParentName);
+	
+	/** Treeview callbacks */
+	TSharedRef<ITableRow> OnGenerateRow(TSharedPtr<ITraceObject> InItem, const TSharedRef<STableViewBase>& OwnerTable);
+	TSharedPtr<SWidget> OnContextMenuOpening() const;
+
+	/** Presets dropdown button callback */
+	TSharedRef<SWidget> MakeAddFilterMenu();
+
+	/** (Re) storing item expansion / selection state */
+	void SaveItemsExpansion();
+	void SaveItemSelection();
+	void RestoreItemsExpansion();	
+	void RestoreItemSelection();
+protected:
+	/** Root level items used by treeview */
+	TArray<TSharedPtr<ITraceObject>> RootItems;
+	/** Dynamically generated array according to treeview filtering */
+	TArray<TSharedPtr<ITraceObject>> TreeItems;
+	/** Flat list of contained treeview items */
+	TArray<TSharedPtr<ITraceObject>> FlatItems;
+	/** Parent child relationships */
+	TMap<TSharedPtr<ITraceObject>, TArray<TSharedPtr<ITraceObject>>> ParentToChild;
+	TMap<TSharedPtr<ITraceObject>, TSharedPtr<ITraceObject>> ChildToParent;
+
+	/** Analysis filter service for the current analysis session this window is currently representing */
+	TSharedPtr<ISessionTraceFilterService> SessionFilterService;
+	/** Timestamp used for refreshing cached filter data */
+	FDateTime SyncTimeStamp;
+
+	FEventFilterService& FilterService;
+	
+	/** Wrapper for presets drop down button */
+	TSharedPtr<SHorizontalBox> OptionsWidget;
+	
+	/** Treeview widget, represents the currently connected analysis session its filtering state */
+	TSharedPtr<STreeView<TSharedPtr<ITraceObject>>> Treeview;
+	TSharedPtr<TreeFilterHandler<TSharedPtr<ITraceObject>>> TreeviewFilterHandler;
+
+	/** Flag indicating the treeview should be refreshed */
+	bool bNeedsTreeRefresh;
+
+	/** Search box for filtering the Treeview items */
+	TSharedPtr<SSearchBox> SearchBoxWidget;
+	/** The text filter used by the search box */
+	TSharedPtr<TTextFilter<TSharedPtr<ITraceObject>>> SearchBoxWidgetFilter;
+
+	/** Filter presets bar widget */
+	TSharedPtr<SFilterPresetList> FilterPresetsListWidget;
+	
+	/** Widget used for selecting specific AnalysisSession the treeview should represent the filter state for */
+	TSharedPtr<SSessionSelector> SessionSelector;
+
+	/** Whether or not we a currently highlighting a preset */
+	bool bHighlightingPreset;
+
+	/** Cached state of which named entries were expanded / selected */
+	TSet<FString> ExpandedObjectNames;
+	TSet<FString> SelectedObjectNames;
+
+	/** Inter-window cached state of last expansion state */
+	static TSet<FString> LastExpandedObjectNames;
+};
\ No newline at end of file
diff --git a/Engine/Plugins/Developer/TraceDataFiltering/Source/TraceDataFiltering/Private/STraceObjectRowWidget.cpp b/Engine/Plugins/Developer/TraceDataFiltering/Source/TraceDataFiltering/Private/STraceObjectRowWidget.cpp
new file mode 100644
--- /dev/null
+++ b/Engine/Plugins/Developer/TraceDataFiltering/Source/TraceDataFiltering/Private/STraceObjectRowWidget.cpp
@@ -0,0 +1,237 @@
+// Copyright Epic Games, Inc. All Rights Reserved.
+
+#include "STraceObjectRowWidget.h"
+#include "Widgets/Views/STableRow.h"
+#include "Widgets/Views/STableViewBase.h"
+#include "Widgets/SBoxPanel.h"
+#include "Widgets/Text/STextBlock.h"
+#include "Widgets/Input/SButton.h"
+#include "EditorStyleSet.h"
+#include "Styling/CoreStyle.h"
+#include "Types/SlateEnums.h"
+
+#include "Async/Async.h"
+
+class FSetFilteringDragDropOp : public FDragDropOperation, public TSharedFromThis<FSetFilteringDragDropOp>
+{
+public:
+	DRAG_DROP_OPERATOR_TYPE(FSetFilteringDragDropOp, FDragDropOperation)
+
+	virtual ~FSetFilteringDragDropOp()
+	{
+		for (TSharedPtr<ITraceObject> Object : Objects)
+		{
+			Object->SetIsFiltered(bFilterOut);
+		}
+	}
+	
+	virtual TSharedPtr<SWidget> GetDefaultDecorator() const override
+	{
+		return SNullWidget::NullWidget;
+	}
+
+	static TSharedRef<FSetFilteringDragDropOp> New(TSharedPtr<ITraceObject> Object)
+	{
+		TSharedRef<FSetFilteringDragDropOp> Operation = MakeShareable(new FSetFilteringDragDropOp);
+		
+		Operation->AddObject(Object);
+		Operation->bFilterOut = !Object->IsFiltered();
+
+		Operation->Construct();
+		return Operation;
+	}
+
+	void AddObject(TSharedPtr<ITraceObject> InObject)
+	{
+		Objects.Add(InObject);
+		InObject->SetPending();
+	}
+
+	bool WillBeFilteredOut() const { return bFilterOut; }
+
+protected:
+	bool bFilterOut;
+	TSet<TSharedPtr<ITraceObject>> Objects;
+};
+
+
+class SFilterCheckboxWidget : public SImage
+{
+public:
+	SLATE_BEGIN_ARGS(SFilterCheckboxWidget) {}
+	SLATE_END_ARGS()
+
+	/** Construct this widget */
+	void Construct(const FArguments& InArgs, TSharedPtr<ITraceObject> InObject)
+	{
+		Object = InObject;
+		SImage::Construct(
+			SImage::FArguments()
+			.Image(this, &SFilterCheckboxWidget::GetBrush)
+			.ColorAndOpacity(FSlateColor::UseForeground())
+		);
+	}
+
+private:
+
+	/** Start a new drag/drop operation for this widget */
+	virtual FReply OnDragDetected(const FGeometry& MyGeometry, const FPointerEvent& MouseEvent) override
+	{
+		if (MouseEvent.IsMouseButtonDown(EKeys::LeftMouseButton))
+		{
+			return FReply::Handled().BeginDragDrop(FSetFilteringDragDropOp::New(Object));
+		}
+		else
+		{
+			return FReply::Unhandled();
+		}
+	}
+
+	/** If a visibility drag drop operation has entered this widget, set its actor to the new visibility state */
+	virtual void OnDragEnter(const FGeometry& MyGeometry, const FDragDropEvent& DragDropEvent) override
+	{
+		auto FilteringOp = DragDropEvent.GetOperationAs<FSetFilteringDragDropOp>();
+		if (FilteringOp.IsValid())
+		{
+			if (Object->IsFiltered() != FilteringOp->WillBeFilteredOut() && !Object->IsPending())
+			{
+				FilteringOp->AddObject(Object);
+			}
+		}
+	}
+
+	FReply HandleClick()
+	{
+		if (!IsPending())
+		{
+			return FReply::Handled().DetectDrag(SharedThis(this), EKeys::LeftMouseButton);
+		}
+
+		return FReply::Handled();		
+	}
+
+	virtual FReply OnMouseButtonDoubleClick(const FGeometry& InMyGeometry, const FPointerEvent& InMouseEvent) override
+	{
+		return HandleClick();
+	}
+
+	virtual FReply OnMouseButtonDown(const FGeometry& MyGeometry, const FPointerEvent& MouseEvent) override
+	{
+		if (MouseEvent.GetEffectingButton() != EKeys::LeftMouseButton)
+		{
+			return FReply::Unhandled();
+		}
+
+		return HandleClick();
+	}
+
+	virtual FReply OnMouseButtonUp(const FGeometry& MyGeometry, const FPointerEvent& MouseEvent) override
+	{
+		if (MouseEvent.GetEffectingButton() == EKeys::LeftMouseButton)
+		{
+			if (!IsPending())
+			{
+				ToggleIsFiltered();
+			}
+
+			return FReply::Handled();
+		}
+
+		return FReply::Unhandled();
+	}
+
+	const FSlateBrush* GetBrush() const
+	{
+		if (IsPending())
+		{
+			return IsHovered() ? FEventFilterStyle::GetBrush("EventFilter.State.Pending_Hovered") :
+				FEventFilterStyle::GetBrush("EventFilter.State.Pending");
+
+		}
+		else
+		{
+			if (IsFiltered())
+			{
+				return IsHovered() ? FEventFilterStyle::GetBrush("EventFilter.State.Disabled_Hovered") :
+					FEventFilterStyle::GetBrush("EventFilter.State.Disabled");
+			}
+			else
+			{
+				return IsHovered() ? FEventFilterStyle::GetBrush("EventFilter.State.Enabled_Hovered") :
+					FEventFilterStyle::GetBrush("EventFilter.State.Enabled");
+			}
+		}
+	}
+
+	bool IsFiltered() const
+	{
+		return Object->IsFiltered();
+	}
+
+	bool IsPending() const
+	{
+		return Object->IsPending();
+	}
+
+	void ToggleIsFiltered()
+	{
+		Object->SetIsFiltered(!Object->IsFiltered());
+	}
+	void SetIsFiltered(bool bState)
+	{
+		Object->SetIsFiltered(bState);
+	}
+
+	TSharedPtr<ITraceObject> Object;
+};
+
+void STraceObjectRowWidget::Construct(const FArguments& InArgs, const TSharedRef<STableViewBase>& InOwnerTableView, TSharedPtr<ITraceObject> InObject)
+{
+	Object = InObject;
+	HighlightTextAttribute = InArgs._HighlightText.Get();
+
+	STableRow<TSharedPtr<ITraceObject>>::Construct(
+		STableRow<TSharedPtr<ITraceObject>>::FArguments()
+		.Padding(0.0f)
+		.Content()
+		[
+			SNew(SHorizontalBox)
+			// Name 
+			+ SHorizontalBox::Slot()
+			.FillWidth(1.0f)
+			.Padding(2, 0, 0, 0)
+			.VAlign(VAlign_Center)
+			[
+				SNew(STextBlock)
+				.Text(Object->GetDisplayText())
+				.HighlightText_Lambda([this]() -> FText
+				{
+					// Get attribute value, populated by search box for treeview filter-ing
+					FText Value = HighlightTextAttribute.Get();
+
+					// In case it is empty allow for treeview highlighting to drive value
+					if (Value.IsEmpty())
+					{
+						auto OwnerTable = OwnerTablePtr.Pin().ToSharedRef();
+						bool bHighlight = OwnerTable->Private_IsItemHighlighted(Object);
+						if (bHighlight)
+						{
+							Value = Object->GetDisplayText();
+						}
+					}
+
+					return Value;
+				})
+			]
+
+			// Visibility icon
+			+ SHorizontalBox::Slot()
+			.AutoWidth()
+			.VAlign(VAlign_Center)
+			.Padding(0, 0, 2.f, 0)
+			[
+				SNew(SFilterCheckboxWidget, InObject)
+			]
+		], InOwnerTableView);
+}
+
diff --git a/Engine/Plugins/Developer/TraceDataFiltering/Source/TraceDataFiltering/Private/STraceObjectRowWidget.h b/Engine/Plugins/Developer/TraceDataFiltering/Source/TraceDataFiltering/Private/STraceObjectRowWidget.h
new file mode 100644
--- /dev/null
+++ b/Engine/Plugins/Developer/TraceDataFiltering/Source/TraceDataFiltering/Private/STraceObjectRowWidget.h
@@ -0,0 +1,25 @@
+// Copyright Epic Games, Inc. All Rights Reserved.
+
+#pragma once
+
+#include "Templates/SharedPointer.h"
+
+class ITraceObject;
+class STableViewBase;
+
+class STraceObjectRowWidget : public STableRow<TSharedPtr<ITraceObject>>
+{
+	SLATE_BEGIN_ARGS(STraceObjectRowWidget) {}
+		/** The current text to highlight */
+		SLATE_ATTRIBUTE(FText, HighlightText)
+	SLATE_END_ARGS()
+
+	void Construct(const FArguments& InArgs, const TSharedRef<STableViewBase>& InOwnerTableView, TSharedPtr<ITraceObject> InObject);
+	virtual ~STraceObjectRowWidget() {}
+
+protected:
+	/* The trace object that this tree item represents */
+	TSharedPtr<ITraceObject> Object;
+	TAttribute<FText> HighlightTextAttribute;
+
+};
diff --git a/Engine/Plugins/Developer/TraceDataFiltering/Source/TraceDataFiltering/Private/Trace/SessionTraceFilterService.cpp b/Engine/Plugins/Developer/TraceDataFiltering/Source/TraceDataFiltering/Private/Trace/SessionTraceFilterService.cpp
new file mode 100644
--- /dev/null
+++ b/Engine/Plugins/Developer/TraceDataFiltering/Source/TraceDataFiltering/Private/Trace/SessionTraceFilterService.cpp
@@ -0,0 +1,137 @@
+// Copyright Epic Games, Inc. All Rights Reserved.
+
+#include "SessionTraceFilterService.h"
+#include "TraceServices/Model/Channel.h"
+#include "TraceServices/ITraceServicesModule.h"
+
+FSessionTraceFilterService::FSessionTraceFilterService(Trace::FSessionHandle InHandle, TSharedPtr<const Trace::IAnalysisSession> InSession) : Session(InSession), Handle(InHandle)
+{
+	FCoreDelegates::OnEndFrame.AddRaw(this, &FSessionTraceFilterService::OnEndFrame);
+}
+
+void FSessionTraceFilterService::GetRootObjects(TArray<FTraceObjectInfo>& OutObjects) const
+{
+	const Trace::IChannelProvider* ChannelProvider = Session->ReadProvider<Trace::IChannelProvider>("ChannelProvider");
+	if (ChannelProvider)
+	{
+		const uint64 ChannelCount = ChannelProvider->GetChannelCount();		
+		const TArray<Trace::FChannelEntry>& Channels = ChannelProvider->GetChannels();
+		for (uint64 ChannelIndex = 0; ChannelIndex < Channels.Num(); ++ChannelIndex)
+		{
+			FTraceObjectInfo& EventInfo = OutObjects.AddDefaulted_GetRef();
+			EventInfo.Name = Channels[ChannelIndex].Name;
+			EventInfo.bEnabled = Channels[ChannelIndex].bIsEnabled;
+			EventInfo.Hash = GetTypeHash(EventInfo.Name);
+			EventInfo.OwnerHash = 0;
+		}
+	}
+}
+
+void FSessionTraceFilterService::GetChildObjects(uint32 InObjectHash, TArray<FTraceObjectInfo>& OutChildObjects) const
+{
+	/** TODO, parent/child relationship for Channels */
+}
+
+const FDateTime& FSessionTraceFilterService::GetTimestamp()
+{
+	const Trace::IChannelProvider* ChannelProvider = Session->ReadProvider<Trace::IChannelProvider>("ChannelProvider");
+	if (ChannelProvider)
+	{
+		TimeStamp = ChannelProvider->GetTimeStamp();
+	}
+
+	return TimeStamp;
+}
+
+void FSessionTraceFilterService::SetObjectFilterState(const FString& InObjectName, const bool bFilterState)
+{
+	if (bFilterState)
+	{
+		FrameDisabledChannels.Remove(InObjectName);
+		FrameEnabledChannels.Add(InObjectName);
+	}
+	else
+	{
+		FrameEnabledChannels.Remove(InObjectName);
+		FrameDisabledChannels.Add(InObjectName);
+	}
+}
+
+void FSessionTraceFilterService::UpdateFilterPresets(const TArray<TSharedPtr<IFilterPreset>>& InPresets)
+{
+	TSet<FString> UniqueNames;
+
+	for (const TSharedPtr<IFilterPreset>& Preset : InPresets)
+	{
+		TArray<FString> Names;
+		Preset->GetWhitelistedNames(Names);
+
+		// We are only interested in the unique channels names, as a result of combining multiple presets
+		Algo::Transform(Names, UniqueNames, [](FString InName) { return InName; });
+	}
+
+	DisableAllChannels();
+
+	for (const FString& EventName : UniqueNames)
+	{
+		SetObjectFilterState(EventName, true);
+	}
+}
+
+void FSessionTraceFilterService::OnEndFrame()
+{
+	ITraceServicesModule& TraceServicesModule = FModuleManager::LoadModuleChecked<ITraceServicesModule>("TraceServices");
+	TSharedPtr<Trace::ISessionService> SessionService = TraceServicesModule.GetSessionService();
+
+	auto GenerateConcatenatedChannels = [](TSet<FString>& InChannels, FString& OutConcatenation)
+	{
+		for (const FString& ChannelName : InChannels)
+		{
+			OutConcatenation += ChannelName;
+			OutConcatenation += TEXT(",");
+		}
+		OutConcatenation.RemoveFromEnd(TEXT(","));
+	};
+
+	if (FrameEnabledChannels.Num())
+	{
+		FString EnabledChannels;
+		GenerateConcatenatedChannels(FrameEnabledChannels, EnabledChannels);
+
+		if (SessionService.IsValid())
+		{
+			UE_LOG(LogTemp, Display, TEXT("CHANNELS %s: %d"), *EnabledChannels, true);
+			SessionService->ToggleChannels(Handle, *EnabledChannels, true);
+		}
+
+		FrameEnabledChannels.Empty();
+	}
+
+	if (FrameDisabledChannels.Num())
+	{
+		FString DisabledChannels;
+		GenerateConcatenatedChannels(FrameDisabledChannels, DisabledChannels);
+				
+		if (SessionService.IsValid())
+		{
+			UE_LOG(LogTemp, Display, TEXT("CHANNELS %s: %d"), *DisabledChannels, false);
+			SessionService->ToggleChannels(Handle, *DisabledChannels, false);
+		}
+
+		FrameDisabledChannels.Empty();
+	}
+}
+
+void FSessionTraceFilterService::DisableAllChannels()
+{
+	const Trace::IChannelProvider* ChannelProvider = Session->ReadProvider<Trace::IChannelProvider>("ChannelProvider");
+	if (ChannelProvider)
+	{
+		const uint64 ChannelCount = ChannelProvider->GetChannelCount();
+		auto Channels = ChannelProvider->GetChannels();
+		for (uint64 ChannelIndex = 0; ChannelIndex < Channels.Num(); ++ChannelIndex)
+		{
+			SetObjectFilterState(Channels[ChannelIndex].Name, false);
+		}
+	}
+}
diff --git a/Engine/Plugins/Developer/TraceDataFiltering/Source/TraceDataFiltering/Private/Trace/SessionTraceFilterService.h b/Engine/Plugins/Developer/TraceDataFiltering/Source/TraceDataFiltering/Private/Trace/SessionTraceFilterService.h
new file mode 100644
--- /dev/null
+++ b/Engine/Plugins/Developer/TraceDataFiltering/Source/TraceDataFiltering/Private/Trace/SessionTraceFilterService.h
@@ -0,0 +1,50 @@
+// Copyright Epic Games, Inc. All Rights Reserved.
+
+#pragma once
+
+#include "ISessionTraceFilterService.h"
+#include "Misc/DateTime.h"
+
+namespace Trace
+{
+	class IAnalysisSession;
+	typedef uint64 FSessionHandle;
+}
+
+class IEventInfoProvider;
+
+/** Implementation of ISessionTraceFilterService specifically to use with Trace, using Trace::IChannelProvider to provide information about Channels available on the running application 
+ and Trace::ISessionService to change the Channel state(s). */
+class FSessionTraceFilterService : public ISessionTraceFilterService
+{
+public:
+	FSessionTraceFilterService(Trace::FSessionHandle InHandle, TSharedPtr<const Trace::IAnalysisSession> InSession);
+	virtual ~FSessionTraceFilterService() {}
+
+	/** Begin ISessionTraceFilterService overrides */
+	virtual void GetRootObjects(TArray<FTraceObjectInfo>& OutObjects) const override;
+	virtual void GetChildObjects(uint32 InObjectHash, TArray<FTraceObjectInfo>& OutChildObjects) const override;
+	virtual const FDateTime& GetTimestamp() override;
+	virtual void SetObjectFilterState(const FString& InObjectName, const bool bFilterState) override;
+	virtual void UpdateFilterPresets(const TArray<TSharedPtr<IFilterPreset>>& InPresets) override;
+	/** End ISessionTraceFilterService overrides */
+
+protected:
+	/** Callback at end of engine frame, used to dispatch all enabled/disabled channels */
+	void OnEndFrame();
+
+	/** Retrieves channels names from provider and marks them all as disabled */
+	void DisableAllChannels();
+protected:
+	/** Session this instance represents the filtering service for */
+	TSharedPtr<const Trace::IAnalysisSession> Session;
+	/** Session handle for AnalyisSession*/
+	Trace::FSessionHandle Handle;
+
+	/** Names of channels that were either enabled or disabled during the duration of this frame */
+	TSet<FString> FrameEnabledChannels;
+	TSet<FString> FrameDisabledChannels;
+	
+	/** Timestamp at which contained data (including provider) was last updated */
+	FDateTime TimeStamp;	
+};
diff --git a/Engine/Plugins/Developer/TraceDataFiltering/Source/TraceDataFiltering/Private/TraceChannel.cpp b/Engine/Plugins/Developer/TraceDataFiltering/Source/TraceDataFiltering/Private/TraceChannel.cpp
new file mode 100644
--- /dev/null
+++ b/Engine/Plugins/Developer/TraceDataFiltering/Source/TraceDataFiltering/Private/TraceChannel.cpp
@@ -0,0 +1,50 @@
+// Copyright Epic Games, Inc. All Rights Reserved.
+
+#include "TraceChannel.h"
+
+FTraceChannel::FTraceChannel(FString InName, FString InParentName, uint32 InHash, bool bInEnabled, const TArray<TSharedPtr<ITraceObject>>& InChildObjects, TSharedPtr<ISessionTraceFilterService> InFilterService) : Name(InName), ParentName(InParentName), Hash(InHash), ChildObjects(InChildObjects), bFiltered(bInEnabled), bIsPending(false), FilterService(InFilterService)
+{
+
+}
+
+FText FTraceChannel::GetDisplayText() const
+{
+	return FText::FromString(Name);
+}
+
+FString FTraceChannel::GetName() const
+{
+	return Name;
+}
+
+void FTraceChannel::SetPending()
+{
+	bIsPending = true;
+}
+
+void FTraceChannel::SetIsFiltered(bool bState)
+{
+	SetPending();
+	FilterService->SetObjectFilterState(Name, !bState);
+}
+
+bool FTraceChannel::IsFiltered() const
+{
+	return bFiltered;
+}
+
+bool FTraceChannel::IsPending() const
+{
+	return bIsPending;
+}
+
+void FTraceChannel::GetSearchString(TArray<FString>& OutFilterStrings) const
+{
+	OutFilterStrings.Add(Name);
+}
+
+void FTraceChannel::GetChildren(TArray<TSharedPtr<ITraceObject>>& OutChildren) const
+{
+	OutChildren.Append(ChildObjects);
+}
+
diff --git a/Engine/Plugins/Developer/TraceDataFiltering/Source/TraceDataFiltering/Private/TraceChannel.h b/Engine/Plugins/Developer/TraceDataFiltering/Source/TraceDataFiltering/Private/TraceChannel.h
new file mode 100644
--- /dev/null
+++ b/Engine/Plugins/Developer/TraceDataFiltering/Source/TraceDataFiltering/Private/TraceChannel.h
@@ -0,0 +1,38 @@
+// Copyright Epic Games, Inc. All Rights Reserved.
+
+#pragma once
+
+#include "ITraceObject.h"
+
+class FTraceChannel : public ITraceObject
+{
+public:
+	FTraceChannel(FString InName, FString InParentName, uint32 InHash, bool bInEnabled, const TArray<TSharedPtr<ITraceObject>>& InChildObjects, TSharedPtr<ISessionTraceFilterService> InFilterService);
+
+	/** Begin ITraceObject overrides */
+	virtual FText GetDisplayText() const override;;
+	virtual FString GetName() const;
+	virtual void SetPending() override;
+	virtual void SetIsFiltered(bool bState) override;
+	virtual bool IsFiltered() const override;
+	virtual bool IsPending() const override;
+	virtual void GetSearchString(TArray<FString>& OutFilterStrings) const override;
+	virtual void GetChildren(TArray<TSharedPtr<ITraceObject>>& OutChildren) const override;
+	/** End ITraceObject overrides */
+
+protected:
+	/** This channel's name */
+	FString Name;
+	/** Channel's parent (group) name */
+	FString ParentName;
+	/** Hash representing this channel uniquely */
+	uint32 Hash;
+	TArray<TSharedPtr<ITraceObject>> ChildObjects;
+	
+	/** Whether or not this channel is filtered out, true = filtered; false = not filtered */
+	bool bFiltered;
+	bool bIsPending;	
+
+	TSharedPtr<ISessionTraceFilterService> FilterService;
+};
+
diff --git a/Engine/Plugins/Developer/TraceDataFiltering/Source/TraceDataFiltering/Private/TraceDataFilteringModule.cpp b/Engine/Plugins/Developer/TraceDataFiltering/Source/TraceDataFiltering/Private/TraceDataFilteringModule.cpp
new file mode 100644
--- /dev/null
+++ b/Engine/Plugins/Developer/TraceDataFiltering/Source/TraceDataFiltering/Private/TraceDataFilteringModule.cpp
@@ -0,0 +1,63 @@
+// Copyright Epic Games, Inc. All Rights Reserved.
+
+#include "TraceDataFilteringModule.h"
+
+#include "Modules/ModuleInterface.h"
+#include "Features/IModularFeatures.h"
+#include "Framework/Docking/LayoutExtender.h"
+#include "Framework/Docking/TabManager.h"
+#include "Templates/SharedPointer.h"
+#include "WorkspaceMenuStructureModule.h"
+#include "WorkspaceMenuStructure.h"
+#include "EditorStyleSet.h"
+
+#include "TraceServices/ModuleService.h"
+#include "TraceServices/AnalysisService.h"
+#include "TraceServices/ITraceServicesModule.h"
+#include "Insights/IUnrealInsightsModule.h"
+
+#include "STraceDataFilterWidget.h"
+
+IMPLEMENT_MODULE(FTraceFilteringModule, TraceDataFiltering);
+
+FName FTraceFilteringModule::InsightsFilterTabName = TEXT("TraceDataFiltering");
+
+FString FTraceFilteringModule::TraceFiltersIni;
+
+void FTraceFilteringModule::StartupModule()
+{
+	FEventFilterStyle::Initialize();
+
+	FConfigCacheIni::LoadGlobalIniFile(TraceFiltersIni, TEXT("TraceDataFilters"));
+
+	FEventFilterService::Get();
+	
+	const FSlateIcon TabIcon(FEventFilterStyle::GetStyleSetName(), "EventFilter.TabIcon");
+	FTabSpawnerEntry& FilterTabSpawnerEntry = FGlobalTabmanager::Get()->RegisterNomadTabSpawner(FTraceFilteringModule::InsightsFilterTabName,
+		FOnSpawnTab::CreateLambda([](const FSpawnTabArgs& Args)
+		{
+			const TSharedRef<SDockTab> DockTab = SNew(SDockTab)
+				.TabRole(ETabRole::NomadTab);
+
+			TSharedRef<STraceDataFilterWidget> Window = SNew(STraceDataFilterWidget);
+			DockTab->SetContent(Window);
+
+			return DockTab;
+	}))
+	.SetDisplayName(NSLOCTEXT("FTraceInsightsModule", "FilteringTabTitle", "Trace Data Filtering"))
+	.SetIcon(TabIcon)
+	.SetTooltipText(NSLOCTEXT("FTraceInsightsModule", "FilteringTabTooltip", "Opens the Trace Data Filtering tab, allows for setting Trace Channel states"));
+
+#if WITH_EDITOR
+	FilterTabSpawnerEntry.SetGroup(WorkspaceMenu::GetMenuStructure().GetDeveloperToolsProfilingCategory());
+#else
+	FilterTabSpawnerEntry.SetGroup(WorkspaceMenu::GetMenuStructure().GetToolsCategory());
+#endif	
+}
+
+void FTraceFilteringModule::ShutdownModule()
+{
+	FEventFilterStyle::Shutdown();
+
+	FGlobalTabmanager::Get()->UnregisterNomadTabSpawner(FTraceFilteringModule::InsightsFilterTabName);
+}
\ No newline at end of file
diff --git a/Engine/Plugins/Developer/TraceDataFiltering/Source/TraceDataFiltering/Private/TraceDataFilteringModule.h b/Engine/Plugins/Developer/TraceDataFiltering/Source/TraceDataFiltering/Private/TraceDataFilteringModule.h
new file mode 100644
--- /dev/null
+++ b/Engine/Plugins/Developer/TraceDataFiltering/Source/TraceDataFiltering/Private/TraceDataFilteringModule.h
@@ -0,0 +1,17 @@
+// Copyright Epic Games, Inc. All Rights Reserved.
+
+#pragma once
+
+#include "CoreMinimal.h"
+
+class FTraceFilteringModule : public IModuleInterface
+{
+public:
+	// IModuleInterface interface
+	virtual void StartupModule() override;
+	virtual void ShutdownModule() override;
+
+	static FString TraceFiltersIni;
+protected:
+	static FName InsightsFilterTabName;	
+};
diff --git a/Engine/Plugins/Developer/TraceDataFiltering/Source/TraceDataFiltering/TraceDataFiltering.build.cs b/Engine/Plugins/Developer/TraceDataFiltering/Source/TraceDataFiltering/TraceDataFiltering.build.cs
new file mode 100644
--- /dev/null
+++ b/Engine/Plugins/Developer/TraceDataFiltering/Source/TraceDataFiltering/TraceDataFiltering.build.cs
@@ -0,0 +1,29 @@
+// Copyright Epic Games, Inc. All Rights Reserved.
+
+using UnrealBuildTool;
+
+public class TraceDataFiltering : ModuleRules
+{
+	public TraceDataFiltering(ReadOnlyTargetRules Target) : base(Target)
+	{
+		PrivateDependencyModuleNames.AddRange(
+                new string[] {
+                    "Core",
+					"CoreUObject",
+                    "SlateCore",
+                    "Slate",
+                    "InputCore",
+					"EditorStyle",                    
+                    "TraceLog",
+					"TraceAnalysis",
+					"TraceServices",
+                    "TraceInsights",
+                }
+            );
+
+		if (Target.bBuildEditor)
+        {
+            PrivateDependencyModuleNames.Add("SharedSettingsWidgets");
+        }
+    }
+}
diff --git a/Engine/Plugins/Developer/TraceDataFiltering/TraceDataFilters.uplugin b/Engine/Plugins/Developer/TraceDataFiltering/TraceDataFilters.uplugin
new file mode 100644
--- /dev/null
+++ b/Engine/Plugins/Developer/TraceDataFiltering/TraceDataFilters.uplugin
@@ -0,0 +1,24 @@
+{
+	"FileVersion" : 1,
+	"Version" : 1,
+	"VersionName" : "1.0",
+	"FriendlyName" : "Trace Data Filtering",
+	"Description" : "Allows for turning on/off individual or sets of Trace Channels.",
+	"Category" : "Editor",
+	"CreatedBy" : "Epic Games, Inc.",
+	"CreatedByURL" : "http://epicgames.com",
+	"DocsURL" : "",
+	"MarketplaceURL" : "",
+	"SupportURL" : "",
+	"EnabledByDefault" : false,
+	"CanContainContent" : false,
+	"IsBetaVersion" : false,
+	"Installed" : false,
+	"Modules": [
+		{
+			"Name": "TraceDataFiltering",
+			"Type": "Editor",
+			"LoadingPhase": "PreDefault"
+		}
+	]
+}
\ No newline at end of file
diff --git a/Engine/Source/Developer/TraceServices/Private/Model/Channel.cpp b/Engine/Source/Developer/TraceServices/Private/Model/Channel.cpp
--- a/Engine/Source/Developer/TraceServices/Private/Model/Channel.cpp
+++ b/Engine/Source/Developer/TraceServices/Private/Model/Channel.cpp
@@ -9,6 +9,7 @@ const FName FChannelProvider::ProviderName("ChannelProvider");
 ///////////////////////////////////////////////////////////////////////////////
 FChannelProvider::FChannelProvider()
 {
+	TimeStamp = FDateTime::Now();
 }
 
 ///////////////////////////////////////////////////////////////////////////////
@@ -21,6 +22,8 @@ void FChannelProvider::AnnounceChannel(const ANSICHAR* InChannelName, uint32 Id)
 		ChannelName,
 		true
 	});
+
+	TimeStamp = FDateTime::Now();
 }
 
 ///////////////////////////////////////////////////////////////////////////////
@@ -34,6 +37,8 @@ void FChannelProvider::UpdateChannel(uint32 Id, bool bEnabled)
 	{
 		FoundEntry->bIsEnabled = bEnabled;
 	}
+
+	TimeStamp = FDateTime::Now();
 }
 
 ///////////////////////////////////////////////////////////////////////////////
@@ -48,4 +53,9 @@ const TArray<FChannelEntry>& FChannelProvider::GetChannels() const
 	return Channels;
 }
 
+FDateTime FChannelProvider::GetTimeStamp() const
+{
+	return TimeStamp;
+}
+
 }
\ No newline at end of file
diff --git a/Engine/Source/Developer/TraceServices/Private/Model/Channel.h b/Engine/Source/Developer/TraceServices/Private/Model/Channel.h
--- a/Engine/Source/Developer/TraceServices/Private/Model/Channel.h
+++ b/Engine/Source/Developer/TraceServices/Private/Model/Channel.h
@@ -18,8 +18,11 @@ public:
 	virtual uint64	GetChannelCount() const override;
 	virtual const TArray<FChannelEntry>& GetChannels() const override;
 
+	virtual FDateTime GetTimeStamp() const override;
+
 private:
 	TArray<FChannelEntry> Channels;
+	FDateTime TimeStamp;
 };
 
 } // namespace Trace
\ No newline at end of file
diff --git a/Engine/Source/Developer/TraceServices/Private/SessionService.cpp b/Engine/Source/Developer/TraceServices/Private/SessionService.cpp
--- a/Engine/Source/Developer/TraceServices/Private/SessionService.cpp
+++ b/Engine/Source/Developer/TraceServices/Private/SessionService.cpp
@@ -159,6 +159,19 @@ void FSessionService::GetLiveSessions(TArray<FSessionHandle>& OutSessions) const
 	}
 }
 
+FSessionHandle FSessionService::GetSessionHandleByName(const FString& SessionName)
+{
+	for (const TPair<FStoreSessionHandle, FSessionInfoInternal>& Pair : Sessions)
+	{
+		if (FString(Pair.Value.Name) == SessionName)
+		{
+			return static_cast<FSessionHandle>(Pair.Key);
+		}
+	}
+
+	return INDEX_NONE;
+}
+
 bool FSessionService::GetSessionInfo(FSessionHandle SessionHandle, FSessionInfo& OutSessionInfo) const
 {
 	FScopeLock Lock(&SessionsCS);
diff --git a/Engine/Source/Developer/TraceServices/Private/SessionServicePrivate.h b/Engine/Source/Developer/TraceServices/Private/SessionServicePrivate.h
--- a/Engine/Source/Developer/TraceServices/Private/SessionServicePrivate.h
+++ b/Engine/Source/Developer/TraceServices/Private/SessionServicePrivate.h
@@ -32,6 +32,7 @@ public:
 	virtual bool IsModuleEnabled(Trace::FSessionHandle SessionHandle, const FName& ModuleName) const override;
 	virtual bool ToggleChannels(Trace::FSessionHandle SessionHandle, const TCHAR* Channels, bool bState) override;
 	virtual bool ConnectSession(const TCHAR* ControlClientAddress) override;
+	virtual FSessionHandle GetSessionHandleByName(const FString& SessionName) override;
 	virtual TSharedPtr<const IAnalysisSession> StartAnalysis(FSessionHandle SessionHandle) override;
 
 private:
diff --git a/Engine/Source/Developer/TraceServices/Public/TraceServices/Model/Channel.h b/Engine/Source/Developer/TraceServices/Public/TraceServices/Model/Channel.h
--- a/Engine/Source/Developer/TraceServices/Public/TraceServices/Model/Channel.h
+++ b/Engine/Source/Developer/TraceServices/Public/TraceServices/Model/Channel.h
@@ -21,6 +21,7 @@ public:
 	virtual ~IChannelProvider() = default;
 	virtual uint64	GetChannelCount() const = 0;
 	virtual const TArray<FChannelEntry>& GetChannels() const = 0;
+	virtual FDateTime GetTimeStamp() const = 0;
 };
 
 } // namespace Trace
\ No newline at end of file
diff --git a/Engine/Source/Developer/TraceServices/Public/TraceServices/SessionService.h b/Engine/Source/Developer/TraceServices/Public/TraceServices/SessionService.h
--- a/Engine/Source/Developer/TraceServices/Public/TraceServices/SessionService.h
+++ b/Engine/Source/Developer/TraceServices/Public/TraceServices/SessionService.h
@@ -43,6 +43,8 @@ public:
 	virtual bool IsModuleEnabled(Trace::FSessionHandle SessionHandle, const FName& ModuleName) const = 0;
 	virtual bool ToggleChannels(Trace::FSessionHandle SessionHandle, const TCHAR* Channels, bool bState) = 0;
 	virtual bool ConnectSession(const TCHAR* ControlClientAddress) = 0;
+	
+	virtual Trace::FSessionHandle GetSessionHandleByName(const FString& SessionName) = 0;
 	virtual TSharedPtr<const IAnalysisSession> StartAnalysis(FSessionHandle SessionHandle) = 0;
 };
 
diff --git a/Engine/Source/Programs/UnrealInsights/Private/UserInterfaceCommand.cpp b/Engine/Source/Programs/UnrealInsights/Private/UserInterfaceCommand.cpp
--- a/Engine/Source/Programs/UnrealInsights/Private/UserInterfaceCommand.cpp
+++ b/Engine/Source/Programs/UnrealInsights/Private/UserInterfaceCommand.cpp
@@ -168,6 +168,7 @@ void FUserInterfaceCommand::Run()
 
 		FStats::AdvanceFrame(false);
 
+		FCoreDelegates::OnEndFrame.Broadcast();
 		GLog->FlushThreadedLogs(); //im: ???
 	}
 
