CIS errors

diff --git a/Engine/Plugins/Developer/TraceDataFiltering/Source/TraceDataFiltering/Private/ISessionTraceFilterService.h b/Engine/Plugins/Developer/TraceDataFiltering/Source/TraceDataFiltering/Private/ISessionTraceFilterService.h
--- a/Engine/Plugins/Developer/TraceDataFiltering/Source/TraceDataFiltering/Private/ISessionTraceFilterService.h
+++ b/Engine/Plugins/Developer/TraceDataFiltering/Source/TraceDataFiltering/Private/ISessionTraceFilterService.h
@@ -4,6 +4,7 @@
 
 #include "Containers/Array.h"
 #include "Containers/UnrealString.h"
+#include "IFilterPreset.h"
 
 struct FTraceObjectInfo
 {
@@ -13,8 +14,6 @@ struct FTraceObjectInfo
 	uint32 OwnerHash;
 };
 
-struct IFilterPreset;
-
 /** Filtering service, representing the state and data for a specific Trace::IAnalysisSession */
 class ISessionTraceFilterService : public TSharedFromThis<ISessionTraceFilterService>
 {
diff --git a/Engine/Source/Developer/TraceServices/Private/Model/Channel.cpp b/Engine/Source/Developer/TraceServices/Private/Model/Channel.cpp
--- a/Engine/Source/Developer/TraceServices/Private/Model/Channel.cpp
+++ b/Engine/Source/Developer/TraceServices/Private/Model/Channel.cpp
@@ -1,6 +1,7 @@
 // Copyright Epic Games, Inc. All Rights Reserved.
 #include "Channel.h"
 #include "UObject/NameTypes.h"
+#include "Misc/DateTime.h"
 
 namespace Trace {
 
