Added some missing #include statements.

#rnx

diff --git a/Engine/Plugins/Developer/TraceDataFiltering/Source/TraceDataFiltering/Private/ISessionTraceFilterService.h b/Engine/Plugins/Developer/TraceDataFiltering/Source/TraceDataFiltering/Private/ISessionTraceFilterService.h
--- a/Engine/Plugins/Developer/TraceDataFiltering/Source/TraceDataFiltering/Private/ISessionTraceFilterService.h
+++ b/Engine/Plugins/Developer/TraceDataFiltering/Source/TraceDataFiltering/Private/ISessionTraceFilterService.h
@@ -5,6 +5,7 @@
 #include "Containers/Array.h"
 #include "Containers/UnrealString.h"
 #include "IFilterPreset.h"
+#include "Templates/SharedPointer.h"
 
 struct FTraceObjectInfo
 {
@@ -34,4 +35,4 @@ public:
 
 	/** Update filtering state according to user-set preset(s) */
 	virtual void UpdateFilterPresets(const TArray<TSharedPtr<IFilterPreset>>& InPresets) = 0;
-};
\ No newline at end of file
+};
diff --git a/Engine/Source/Developer/TraceServices/Private/Model/Channel.h b/Engine/Source/Developer/TraceServices/Private/Model/Channel.h
--- a/Engine/Source/Developer/TraceServices/Private/Model/Channel.h
+++ b/Engine/Source/Developer/TraceServices/Private/Model/Channel.h
@@ -2,6 +2,8 @@
 
 #pragma once
 
+#include "Containers/Array.h"
+#include "Misc/DateTime.h"
 #include "TraceServices/Model/Channel.h"
 
 namespace Trace {
@@ -25,4 +27,4 @@ private:
 	FDateTime TimeStamp;
 };
 
-} // namespace Trace
\ No newline at end of file
+} // namespace Trace
