Boatload of CIS PCH / unity compile fixes

diff --git a/Engine/Plugins/Developer/TraceDataFiltering/Source/TraceDataFiltering/Private/FilterPresets.cpp b/Engine/Plugins/Developer/TraceDataFiltering/Source/TraceDataFiltering/Private/FilterPresets.cpp
--- a/Engine/Plugins/Developer/TraceDataFiltering/Source/TraceDataFiltering/Private/FilterPresets.cpp
+++ b/Engine/Plugins/Developer/TraceDataFiltering/Source/TraceDataFiltering/Private/FilterPresets.cpp
@@ -1,6 +1,8 @@
 // Copyright Epic Games, Inc. All Rights Reserved.
 
 #include "FilterPresets.h"
+#include "HAL/PlatformFilemanager.h"
+#include "GenericPlatform/GenericPlatformFile.h"
 
 #define LOCTEXT_NAMESPACE "FilterPreset"
 
diff --git a/Engine/Plugins/Developer/TraceDataFiltering/Source/TraceDataFiltering/Private/IFilterPreset.h b/Engine/Plugins/Developer/TraceDataFiltering/Source/TraceDataFiltering/Private/IFilterPreset.h
--- a/Engine/Plugins/Developer/TraceDataFiltering/Source/TraceDataFiltering/Private/IFilterPreset.h
+++ b/Engine/Plugins/Developer/TraceDataFiltering/Source/TraceDataFiltering/Private/IFilterPreset.h
@@ -2,6 +2,8 @@
 
 #pragma once
 
+#include "Templates/SharedPointer.h"
+
 class ITraceObject;
 
 /** Interface for structures representing a filtering preset */
diff --git a/Engine/Plugins/Developer/TraceDataFiltering/Source/TraceDataFiltering/Private/ITraceObject.h b/Engine/Plugins/Developer/TraceDataFiltering/Source/TraceDataFiltering/Private/ITraceObject.h
--- a/Engine/Plugins/Developer/TraceDataFiltering/Source/TraceDataFiltering/Private/ITraceObject.h
+++ b/Engine/Plugins/Developer/TraceDataFiltering/Source/TraceDataFiltering/Private/ITraceObject.h
@@ -4,6 +4,8 @@
 
 #include "Templates/SharedPointer.h"
 #include "Containers/Array.h"
+#include "Containers/UnrealString.h"
+#include "Internationalization/Text.h"
 
 class ITraceObject : public TSharedFromThis<ITraceObject>
 {
diff --git a/Engine/Plugins/Developer/TraceDataFiltering/Source/TraceDataFiltering/Private/SFilterPreset.cpp b/Engine/Plugins/Developer/TraceDataFiltering/Source/TraceDataFiltering/Private/SFilterPreset.cpp
--- a/Engine/Plugins/Developer/TraceDataFiltering/Source/TraceDataFiltering/Private/SFilterPreset.cpp
+++ b/Engine/Plugins/Developer/TraceDataFiltering/Source/TraceDataFiltering/Private/SFilterPreset.cpp
@@ -2,6 +2,10 @@
 
 #include "SFilterPreset.h"
 #include "EditorStyleSet.h"
+#include "Widgets/Layout/SBorder.h"
+#include "Framework/MultiBox/MultiBoxBuilder.h"
+
+#include "EventFilterStyle.h"
 
 #define LOCTEXT_NAMESPACE "SFilterPreset"
 
diff --git a/Engine/Plugins/Developer/TraceDataFiltering/Source/TraceDataFiltering/Private/SFilterPreset.h b/Engine/Plugins/Developer/TraceDataFiltering/Source/TraceDataFiltering/Private/SFilterPreset.h
--- a/Engine/Plugins/Developer/TraceDataFiltering/Source/TraceDataFiltering/Private/SFilterPreset.h
+++ b/Engine/Plugins/Developer/TraceDataFiltering/Source/TraceDataFiltering/Private/SFilterPreset.h
@@ -3,6 +3,8 @@
 #pragma once
 
 #include "CoreMinimal.h"
+#include "Widgets/Input/SCheckBox.h"
+#include "Widgets/SCompoundWidget.h"
 
 DECLARE_DELEGATE(FOnPresetsChanged);
 
diff --git a/Engine/Plugins/Developer/TraceDataFiltering/Source/TraceDataFiltering/Private/SFilterPresetList.cpp b/Engine/Plugins/Developer/TraceDataFiltering/Source/TraceDataFiltering/Private/SFilterPresetList.cpp
--- a/Engine/Plugins/Developer/TraceDataFiltering/Source/TraceDataFiltering/Private/SFilterPresetList.cpp
+++ b/Engine/Plugins/Developer/TraceDataFiltering/Source/TraceDataFiltering/Private/SFilterPresetList.cpp
@@ -14,18 +14,21 @@
 #include "Widgets/Input/SEditableText.h"
 #include "Widgets/Text/STextBlock.h"
 #include "Widgets/Input/SCheckBox.h"
-#include "Widgets/Text/SInlineEditableTextBlock.h"
+#include "Widgets/Input/SEditableTextBox.h"
 #include "Styling/SlateTypes.h"
 #include "Misc/ConfigCacheIni.h"
 #include "Layout/WidgetPath.h"
 #include "Textures/SlateIcon.h"
 #include "EditorStyleSet.h"
 
+#include "SFilterPreset.h"
 #include "TraceDataFilteringModule.h"
+#include "FilterPresets.h"
 
 #if WITH_EDITOR
 #include "SSettingsEditorCheckoutNotice.h"
 #endif
+#include "Algo/Transform.h"
 
 #define LOCTEXT_NAMESPACE "InsightsFilterList"
 
diff --git a/Engine/Plugins/Developer/TraceDataFiltering/Source/TraceDataFiltering/Private/SFilterPresetList.h b/Engine/Plugins/Developer/TraceDataFiltering/Source/TraceDataFiltering/Private/SFilterPresetList.h
--- a/Engine/Plugins/Developer/TraceDataFiltering/Source/TraceDataFiltering/Private/SFilterPresetList.h
+++ b/Engine/Plugins/Developer/TraceDataFiltering/Source/TraceDataFiltering/Private/SFilterPresetList.h
@@ -7,10 +7,15 @@
 #include "Widgets/DeclarativeSyntaxSupport.h"
 #include "Widgets/SWidget.h"
 #include "Widgets/SCompoundWidget.h"
+#include "Templates/SharedPointer.h"
+
+#include "SFilterPreset.h"
 
 class FMenuBuilder;
 class SFilterPreset;
 class SWrapBox;
+class ULocalFilterPresetContainer;
+class USharedFilterPresetContainer;
 
 /**
  * A list of presets currently applied to insights filtering window
diff --git a/Engine/Plugins/Developer/TraceDataFiltering/Source/TraceDataFiltering/Private/STraceDataFilterWidget.cpp b/Engine/Plugins/Developer/TraceDataFiltering/Source/TraceDataFiltering/Private/STraceDataFilterWidget.cpp
--- a/Engine/Plugins/Developer/TraceDataFiltering/Source/TraceDataFiltering/Private/STraceDataFilterWidget.cpp
+++ b/Engine/Plugins/Developer/TraceDataFiltering/Source/TraceDataFiltering/Private/STraceDataFilterWidget.cpp
@@ -6,12 +6,18 @@
 #include "Widgets/Text/STextBlock.h"
 #include "Widgets/Input/SComboButton.h"
 #include "Widgets/Input/SSearchBox.h"
+#include "Widgets/SOverlay.h"
 #include "Templates/SharedPointer.h"
-#include "SFilterPresetList.h"
-#include "STraceObjectRowWidget.h"
+#include "TraceServices/ITraceServicesModule.h"
+#include "TraceServices/SessionService.h"
+#include "Framework/MultiBox/MultiBoxBuilder.h"
+#include "Modules/ModuleManager.h"
 
+#include "STraceObjectRowWidget.h"
+#include "SFilterPresetList.h"
 #include "ITraceObject.h"
 #include "TraceChannel.h"
+#include "FilterPresets.h"
 
 #define LOCTEXT_NAMESPACE "STraceDataFilterWidget"
 
diff --git a/Engine/Plugins/Developer/TraceDataFiltering/Source/TraceDataFiltering/Private/STraceDataFilterWidget.h b/Engine/Plugins/Developer/TraceDataFiltering/Source/TraceDataFiltering/Private/STraceDataFilterWidget.h
--- a/Engine/Plugins/Developer/TraceDataFiltering/Source/TraceDataFiltering/Private/STraceDataFilterWidget.h
+++ b/Engine/Plugins/Developer/TraceDataFiltering/Source/TraceDataFiltering/Private/STraceDataFilterWidget.h
@@ -11,6 +11,7 @@
 class ITableRow;
 class SFilterPresetList;
 class SSessionSelector;
+class SSearchBox;
 
 class STraceDataFilterWidget : public SCompoundWidget
 {
diff --git a/Engine/Plugins/Developer/TraceDataFiltering/Source/TraceDataFiltering/Private/STraceObjectRowWidget.cpp b/Engine/Plugins/Developer/TraceDataFiltering/Source/TraceDataFiltering/Private/STraceObjectRowWidget.cpp
--- a/Engine/Plugins/Developer/TraceDataFiltering/Source/TraceDataFiltering/Private/STraceObjectRowWidget.cpp
+++ b/Engine/Plugins/Developer/TraceDataFiltering/Source/TraceDataFiltering/Private/STraceObjectRowWidget.cpp
@@ -9,9 +9,15 @@
 #include "EditorStyleSet.h"
 #include "Styling/CoreStyle.h"
 #include "Types/SlateEnums.h"
-
 #include "Async/Async.h"
 
+#include "ITraceObject.h"
+#include "EventFilterStyle.h"
+#include "Widgets/Images/SImage.h"
+#include "Templates/SharedPointer.h"
+#include "Input/Events.h"
+#include "Layout/Geometry.h"
+
 class FSetFilteringDragDropOp : public FDragDropOperation, public TSharedFromThis<FSetFilteringDragDropOp>
 {
 public:
diff --git a/Engine/Plugins/Developer/TraceDataFiltering/Source/TraceDataFiltering/Private/STraceObjectRowWidget.h b/Engine/Plugins/Developer/TraceDataFiltering/Source/TraceDataFiltering/Private/STraceObjectRowWidget.h
--- a/Engine/Plugins/Developer/TraceDataFiltering/Source/TraceDataFiltering/Private/STraceObjectRowWidget.h
+++ b/Engine/Plugins/Developer/TraceDataFiltering/Source/TraceDataFiltering/Private/STraceObjectRowWidget.h
@@ -3,8 +3,15 @@
 #pragma once
 
 #include "Templates/SharedPointer.h"
+#include "Misc/Attribute.h"
+#include "Internationalization/Text.h"
+#include "Widgets/Views/STableRow.h"
+#include "Widgets/Views/STableViewBase.h"
+#include "Widgets/Views/SListView.h"
+#include "Widgets/DeclarativeSyntaxSupport.h" 
+
+#include "ITraceObject.h"
 
-class ITraceObject;
 class STableViewBase;
 
 class STraceObjectRowWidget : public STableRow<TSharedPtr<ITraceObject>>
diff --git a/Engine/Plugins/Developer/TraceDataFiltering/Source/TraceDataFiltering/Private/Trace/SessionTraceFilterService.cpp b/Engine/Plugins/Developer/TraceDataFiltering/Source/TraceDataFiltering/Private/Trace/SessionTraceFilterService.cpp
--- a/Engine/Plugins/Developer/TraceDataFiltering/Source/TraceDataFiltering/Private/Trace/SessionTraceFilterService.cpp
+++ b/Engine/Plugins/Developer/TraceDataFiltering/Source/TraceDataFiltering/Private/Trace/SessionTraceFilterService.cpp
@@ -3,6 +3,11 @@
 #include "SessionTraceFilterService.h"
 #include "TraceServices/Model/Channel.h"
 #include "TraceServices/ITraceServicesModule.h"
+#include "Templates/SharedPointer.h"
+#include "TraceServices/SessionService.h"
+#include "Misc/CoreDelegates.h"
+#include "Algo/Transform.h"
+#include "Modules/ModuleManager.h"
 
 FSessionTraceFilterService::FSessionTraceFilterService(Trace::FSessionHandle InHandle, TSharedPtr<const Trace::IAnalysisSession> InSession) : Session(InSession), Handle(InHandle)
 {
diff --git a/Engine/Plugins/Developer/TraceDataFiltering/Source/TraceDataFiltering/Private/TraceChannel.cpp b/Engine/Plugins/Developer/TraceDataFiltering/Source/TraceDataFiltering/Private/TraceChannel.cpp
--- a/Engine/Plugins/Developer/TraceDataFiltering/Source/TraceDataFiltering/Private/TraceChannel.cpp
+++ b/Engine/Plugins/Developer/TraceDataFiltering/Source/TraceDataFiltering/Private/TraceChannel.cpp
@@ -1,6 +1,7 @@
 // Copyright Epic Games, Inc. All Rights Reserved.
 
 #include "TraceChannel.h"
+#include "Internationalization/Text.h"
 
 FTraceChannel::FTraceChannel(FString InName, FString InParentName, uint32 InHash, bool bInEnabled, const TArray<TSharedPtr<ITraceObject>>& InChildObjects, TSharedPtr<ISessionTraceFilterService> InFilterService) : Name(InName), ParentName(InParentName), Hash(InHash), ChildObjects(InChildObjects), bFiltered(bInEnabled), bIsPending(false), FilterService(InFilterService)
 {
diff --git a/Engine/Plugins/Developer/TraceDataFiltering/Source/TraceDataFiltering/Private/TraceDataFilteringModule.cpp b/Engine/Plugins/Developer/TraceDataFiltering/Source/TraceDataFiltering/Private/TraceDataFilteringModule.cpp
--- a/Engine/Plugins/Developer/TraceDataFiltering/Source/TraceDataFiltering/Private/TraceDataFilteringModule.cpp
+++ b/Engine/Plugins/Developer/TraceDataFiltering/Source/TraceDataFiltering/Private/TraceDataFilteringModule.cpp
@@ -3,6 +3,7 @@
 #include "TraceDataFilteringModule.h"
 
 #include "Modules/ModuleInterface.h"
+#include "Modules/ModuleManager.h"
 #include "Features/IModularFeatures.h"
 #include "Framework/Docking/LayoutExtender.h"
 #include "Framework/Docking/TabManager.h"
@@ -10,6 +11,8 @@
 #include "WorkspaceMenuStructureModule.h"
 #include "WorkspaceMenuStructure.h"
 #include "EditorStyleSet.h"
+#include "Widgets/Docking/SDockTab.h"
+#include "Misc/ConfigCacheIni.h"
 
 #include "TraceServices/ModuleService.h"
 #include "TraceServices/AnalysisService.h"
@@ -17,6 +20,7 @@
 #include "Insights/IUnrealInsightsModule.h"
 
 #include "STraceDataFilterWidget.h"
+#include "EventFilterStyle.h"
 
 IMPLEMENT_MODULE(FTraceFilteringModule, TraceDataFiltering);
 
diff --git a/Engine/Plugins/Developer/TraceDataFiltering/Source/TraceDataFiltering/Private/TraceDataFilteringModule.h b/Engine/Plugins/Developer/TraceDataFiltering/Source/TraceDataFiltering/Private/TraceDataFilteringModule.h
--- a/Engine/Plugins/Developer/TraceDataFiltering/Source/TraceDataFiltering/Private/TraceDataFilteringModule.h
+++ b/Engine/Plugins/Developer/TraceDataFiltering/Source/TraceDataFiltering/Private/TraceDataFilteringModule.h
@@ -3,6 +3,7 @@
 #pragma once
 
 #include "CoreMinimal.h"
+#include "Modules/ModuleInterface.h"
 
 class FTraceFilteringModule : public IModuleInterface
 {
