More CIS include / pre-declare fixes

diff --git a/Engine/Plugins/Developer/TraceDataFiltering/Source/TraceDataFiltering/Private/SFilterPreset.cpp b/Engine/Plugins/Developer/TraceDataFiltering/Source/TraceDataFiltering/Private/SFilterPreset.cpp
--- a/Engine/Plugins/Developer/TraceDataFiltering/Source/TraceDataFiltering/Private/SFilterPreset.cpp
+++ b/Engine/Plugins/Developer/TraceDataFiltering/Source/TraceDataFiltering/Private/SFilterPreset.cpp
@@ -6,6 +6,7 @@
 #include "Framework/MultiBox/MultiBoxBuilder.h"
 
 #include "EventFilterStyle.h"
+#include "IFilterPreset.h"
 
 #define LOCTEXT_NAMESPACE "SFilterPreset"
 
diff --git a/Engine/Plugins/Developer/TraceDataFiltering/Source/TraceDataFiltering/Private/SFilterPreset.h b/Engine/Plugins/Developer/TraceDataFiltering/Source/TraceDataFiltering/Private/SFilterPreset.h
--- a/Engine/Plugins/Developer/TraceDataFiltering/Source/TraceDataFiltering/Private/SFilterPreset.h
+++ b/Engine/Plugins/Developer/TraceDataFiltering/Source/TraceDataFiltering/Private/SFilterPreset.h
@@ -5,6 +5,9 @@
 #include "CoreMinimal.h"
 #include "Widgets/Input/SCheckBox.h"
 #include "Widgets/SCompoundWidget.h"
+#include "Delegates/DelegateCombinations.h"
+
+struct IFilterPreset;
 
 DECLARE_DELEGATE(FOnPresetsChanged);
 
diff --git a/Engine/Plugins/Developer/TraceDataFiltering/Source/TraceDataFiltering/Private/STraceDataFilterWidget.cpp b/Engine/Plugins/Developer/TraceDataFiltering/Source/TraceDataFiltering/Private/STraceDataFilterWidget.cpp
--- a/Engine/Plugins/Developer/TraceDataFiltering/Source/TraceDataFiltering/Private/STraceDataFilterWidget.cpp
+++ b/Engine/Plugins/Developer/TraceDataFiltering/Source/TraceDataFiltering/Private/STraceDataFilterWidget.cpp
@@ -18,6 +18,7 @@
 #include "ITraceObject.h"
 #include "TraceChannel.h"
 #include "FilterPresets.h"
+#include "IFilterPreset.h"
 
 #define LOCTEXT_NAMESPACE "STraceDataFilterWidget"
 
diff --git a/Engine/Plugins/Developer/TraceDataFiltering/Source/TraceDataFiltering/Private/STraceDataFilterWidget.h b/Engine/Plugins/Developer/TraceDataFiltering/Source/TraceDataFiltering/Private/STraceDataFilterWidget.h
--- a/Engine/Plugins/Developer/TraceDataFiltering/Source/TraceDataFiltering/Private/STraceDataFilterWidget.h
+++ b/Engine/Plugins/Developer/TraceDataFiltering/Source/TraceDataFiltering/Private/STraceDataFilterWidget.h
@@ -3,15 +3,23 @@
 #pragma once
 
 #include "Templates/SharedPointer.h"
-#include "ITraceObject.h"
+#include "TraceServices/SessionService.h"
+#include "Containers/UnrealString.h"
 #include "Misc/IFilter.h"
 #include "Framework/Views/TreeFilterHandler.h"
 #include "Misc/TextFilter.h"
 
+#include "ITraceObject.h"
+
 class ITableRow;
 class SFilterPresetList;
 class SSessionSelector;
 class SSearchBox;
+struct IFilterPreset;
+class ITraceObject;
+class ISessionTraceFilterService;
+class FEventFilterService;
+struct FTraceObjectInfo;
 
 class STraceDataFilterWidget : public SCompoundWidget
 {
diff --git a/Engine/Plugins/Developer/TraceDataFiltering/Source/TraceDataFiltering/Private/TraceChannel.cpp b/Engine/Plugins/Developer/TraceDataFiltering/Source/TraceDataFiltering/Private/TraceChannel.cpp
--- a/Engine/Plugins/Developer/TraceDataFiltering/Source/TraceDataFiltering/Private/TraceChannel.cpp
+++ b/Engine/Plugins/Developer/TraceDataFiltering/Source/TraceDataFiltering/Private/TraceChannel.cpp
@@ -2,6 +2,7 @@
 
 #include "TraceChannel.h"
 #include "Internationalization/Text.h"
+#include "ISessionTraceFilterService.h"
 
 FTraceChannel::FTraceChannel(FString InName, FString InParentName, uint32 InHash, bool bInEnabled, const TArray<TSharedPtr<ITraceObject>>& InChildObjects, TSharedPtr<ISessionTraceFilterService> InFilterService) : Name(InName), ParentName(InParentName), Hash(InHash), ChildObjects(InChildObjects), bFiltered(bInEnabled), bIsPending(false), FilterService(InFilterService)
 {
diff --git a/Engine/Plugins/Developer/TraceDataFiltering/Source/TraceDataFiltering/Private/TraceChannel.h b/Engine/Plugins/Developer/TraceDataFiltering/Source/TraceDataFiltering/Private/TraceChannel.h
--- a/Engine/Plugins/Developer/TraceDataFiltering/Source/TraceDataFiltering/Private/TraceChannel.h
+++ b/Engine/Plugins/Developer/TraceDataFiltering/Source/TraceDataFiltering/Private/TraceChannel.h
@@ -4,6 +4,8 @@
 
 #include "ITraceObject.h"
 
+class ISessionTraceFilterService;
+
 class FTraceChannel : public ITraceObject
 {
 public:
diff --git a/Engine/Plugins/Developer/TraceDataFiltering/Source/TraceDataFiltering/Private/TraceDataFilteringModule.cpp b/Engine/Plugins/Developer/TraceDataFiltering/Source/TraceDataFiltering/Private/TraceDataFilteringModule.cpp
--- a/Engine/Plugins/Developer/TraceDataFiltering/Source/TraceDataFiltering/Private/TraceDataFilteringModule.cpp
+++ b/Engine/Plugins/Developer/TraceDataFiltering/Source/TraceDataFiltering/Private/TraceDataFilteringModule.cpp
@@ -21,6 +21,7 @@
 
 #include "STraceDataFilterWidget.h"
 #include "EventFilterStyle.h"
+#include "EventFilterService.h"
 
 IMPLEMENT_MODULE(FTraceFilteringModule, TraceDataFiltering);
 
