Fix up Insights Channel filtering to work with new session store system

diff --git a/Engine/Plugins/Developer/TraceDataFiltering/Source/TraceDataFiltering/Private/EventFilterService.cpp b/Engine/Plugins/Developer/TraceDataFiltering/Source/TraceDataFiltering/Private/EventFilterService.cpp
deleted file mode 100644
diff --git a/Engine/Plugins/Developer/TraceDataFiltering/Source/TraceDataFiltering/Private/EventFilterService.h b/Engine/Plugins/Developer/TraceDataFiltering/Source/TraceDataFiltering/Private/EventFilterService.h
deleted file mode 100644
diff --git a/Engine/Plugins/Developer/TraceDataFiltering/Source/TraceDataFiltering/Private/FilterPresets.h b/Engine/Plugins/Developer/TraceDataFiltering/Source/TraceDataFiltering/Private/FilterPresets.h
--- a/Engine/Plugins/Developer/TraceDataFiltering/Source/TraceDataFiltering/Private/FilterPresets.h
+++ b/Engine/Plugins/Developer/TraceDataFiltering/Source/TraceDataFiltering/Private/FilterPresets.h
@@ -6,7 +6,6 @@
 #include "UObject/ObjectMacros.h"
 #include "Containers/Map.h"
 #include "Containers/UnrealString.h"
-#include "EventFilterService.h"
 #include "ITraceObject.h"
 
 #include "FilterPresets.generated.h"
diff --git a/Engine/Plugins/Developer/TraceDataFiltering/Source/TraceDataFiltering/Private/STraceDataFilterWidget.cpp b/Engine/Plugins/Developer/TraceDataFiltering/Source/TraceDataFiltering/Private/STraceDataFilterWidget.cpp
--- a/Engine/Plugins/Developer/TraceDataFiltering/Source/TraceDataFiltering/Private/STraceDataFilterWidget.cpp
+++ b/Engine/Plugins/Developer/TraceDataFiltering/Source/TraceDataFiltering/Private/STraceDataFilterWidget.cpp
@@ -19,12 +19,15 @@
 #include "TraceChannel.h"
 #include "FilterPresets.h"
 #include "IFilterPreset.h"
+#include "Insights/IUnrealInsightsModule.h"
+#include "Trace/StoreClient.h"
+#include "Trace/SessionTraceFilterService.h"
 
 #define LOCTEXT_NAMESPACE "STraceDataFilterWidget"
 
 TSet<FString> STraceDataFilterWidget::LastExpandedObjectNames;
 
-STraceDataFilterWidget::STraceDataFilterWidget() : FilterService(FEventFilterService::Get()), bNeedsTreeRefresh(false), bHighlightingPreset(false)
+STraceDataFilterWidget::STraceDataFilterWidget() : bNeedsTreeRefresh(false), bHighlightingPreset(false)
 {
 }
 
@@ -531,9 +534,10 @@ void STraceDataFilterWidget::RestoreItemSelection()
 	SelectedObjectNames.Empty();
 }
 
-void STraceDataFilterWidget::SetCurrentAnalysisSession(Trace::FSessionHandle Handle)
+void STraceDataFilterWidget::SetCurrentAnalysisSession(uint32 SessionHandle, TSharedRef<const Trace::IAnalysisSession> AnalysisSession)
 {
-	SessionFilterService = FilterService.GetFilterServiceByHandle(Handle);
+	SessionFilterService = MakeShareable(new FSessionTraceFilterService(SessionHandle, AnalysisSession));
+
 	/** Refresh presets so config loaded state is directly applied */
 	OnPresetsChanged();
 	/** Refresh data driving the treeview */
@@ -678,13 +682,17 @@ void STraceDataFilterWidget::Tick(const FGeometry& AllottedGeometry, const doubl
 	}
 	else
 	{
-		ITraceServicesModule& TraceServicesModule = FModuleManager::LoadModuleChecked<ITraceServicesModule>("TraceServices");
-		TArray<Trace::FSessionHandle> LiveSessionHandles;
-		TraceServicesModule.GetSessionService()->GetLiveSessions(LiveSessionHandles);
-		for (const Trace::FSessionHandle& Handle : LiveSessionHandles)
+		IUnrealInsightsModule& InsightsModule = FModuleManager::LoadModuleChecked<IUnrealInsightsModule>("TraceInsights");
+		Trace::FStoreClient* StoreClient = InsightsModule.GetStoreClient();
+		const int32 SessionCount = StoreClient->GetSessionCount();
+
+		if (SessionCount > 0)
 		{
-			SetCurrentAnalysisSession(Handle);
-			break;
+			const Trace::FStoreClient::FSessionInfo* SessionInfo = StoreClient->GetSessionInfo(SessionCount - 1);
+			if (SessionInfo)
+			{
+				SetCurrentAnalysisSession(SessionInfo->GetTraceId(), InsightsModule.GetAnalysisSession().ToSharedRef());
+			}
 		}
 	}
 
diff --git a/Engine/Plugins/Developer/TraceDataFiltering/Source/TraceDataFiltering/Private/STraceDataFilterWidget.h b/Engine/Plugins/Developer/TraceDataFiltering/Source/TraceDataFiltering/Private/STraceDataFilterWidget.h
--- a/Engine/Plugins/Developer/TraceDataFiltering/Source/TraceDataFiltering/Private/STraceDataFilterWidget.h
+++ b/Engine/Plugins/Developer/TraceDataFiltering/Source/TraceDataFiltering/Private/STraceDataFilterWidget.h
@@ -53,8 +53,8 @@ protected:
 	void ConstructSearchBoxFilter();
 	void ConstructFilterHandler();
 
-	/** Callback for whenever a different analysis session (by handle), has been selected through the SessionSelector widget */
-	void SetCurrentAnalysisSession(Trace::FSessionHandle Handle);
+	/** Callback for whenever a different analysis session (store) has been retrieved */
+	void SetCurrentAnalysisSession(uint32 SessionHandle, TSharedRef<const Trace::IAnalysisSession> AnalysisSession);
 	bool HasValidFilterSession() const;
 		
 	void RefreshTreeviewData();
@@ -99,8 +99,6 @@ protected:
 	TSharedPtr<ISessionTraceFilterService> SessionFilterService;
 	/** Timestamp used for refreshing cached filter data */
 	FDateTime SyncTimeStamp;
-
-	FEventFilterService& FilterService;
 	
 	/** Wrapper for presets drop down button */
 	TSharedPtr<SHorizontalBox> OptionsWidget;
diff --git a/Engine/Plugins/Developer/TraceDataFiltering/Source/TraceDataFiltering/Private/Trace/SessionTraceFilterService.cpp b/Engine/Plugins/Developer/TraceDataFiltering/Source/TraceDataFiltering/Private/Trace/SessionTraceFilterService.cpp
--- a/Engine/Plugins/Developer/TraceDataFiltering/Source/TraceDataFiltering/Private/Trace/SessionTraceFilterService.cpp
+++ b/Engine/Plugins/Developer/TraceDataFiltering/Source/TraceDataFiltering/Private/Trace/SessionTraceFilterService.cpp
@@ -8,6 +8,11 @@
 #include "Misc/CoreDelegates.h"
 #include "Algo/Transform.h"
 #include "Modules/ModuleManager.h"
+#include <Insights/IUnrealInsightsModule.h>
+#include <Trace/StoreClient.h>
+#include <Trace/ControlClient.h>
+#include <IPAddress.h>
+#include <SocketSubsystem.h>
 
 FSessionTraceFilterService::FSessionTraceFilterService(Trace::FSessionHandle InHandle, TSharedPtr<const Trace::IAnalysisSession> InSession) : Session(InSession), Handle(InHandle)
 {
@@ -85,9 +90,6 @@ void FSessionTraceFilterService::UpdateFilterPresets(const TArray<TSharedPtr<IFi
 
 void FSessionTraceFilterService::OnEndFrame()
 {
-	ITraceServicesModule& TraceServicesModule = FModuleManager::LoadModuleChecked<ITraceServicesModule>("TraceServices");
-	TSharedPtr<Trace::ISessionService> SessionService = TraceServicesModule.GetSessionService();
-
 	auto GenerateConcatenatedChannels = [](TSet<FString>& InChannels, FString& OutConcatenation)
 	{
 		for (const FString& ChannelName : InChannels)
@@ -98,17 +100,40 @@ void FSessionTraceFilterService::OnEndFrame()
 		OutConcatenation.RemoveFromEnd(TEXT(","));
 	};
 
+	if (FrameEnabledChannels.Num() == 0 && FrameDisabledChannels.Num() == 0)
+	{
+		return;
+	}
+
+	IUnrealInsightsModule& InsightsModule = FModuleManager::LoadModuleChecked<IUnrealInsightsModule>("TraceInsights");
+	Trace::FStoreClient* StoreClient = InsightsModule.GetStoreClient();
+
+	if (!StoreClient)
+	{
+		return;
+	}
+
+	const Trace::FStoreClient::FSessionInfo* SessionInfo = StoreClient->GetSessionInfoByTraceId(Handle);
+
+	ISocketSubsystem* Sockets = ISocketSubsystem::Get();
+	TSharedRef<FInternetAddr> ClientAddr(Sockets->CreateInternetAddr());
+	ClientAddr->SetIp(SessionInfo->GetIpAddress());
+	ClientAddr->SetPort(1985);
+
+	Trace::FControlClient ControlClient;
+	if (!ControlClient.Connect(ClientAddr.Get()))
+	{
+		return;
+	}
+
 	if (FrameEnabledChannels.Num())
 	{
 		FString EnabledChannels;
 		GenerateConcatenatedChannels(FrameEnabledChannels, EnabledChannels);
 
-		if (SessionService.IsValid())
-		{
-			UE_LOG(LogTemp, Display, TEXT("CHANNELS %s: %d"), *EnabledChannels, true);
-			SessionService->ToggleChannels(Handle, *EnabledChannels, true);
-		}
-
+		UE_LOG(LogTemp, Display, TEXT("CHANNELS %s: %d"), *EnabledChannels, true);
+		ControlClient.SendToggleChannel(*EnabledChannels, true);
+		
 		FrameEnabledChannels.Empty();
 	}
 
@@ -117,14 +142,13 @@ void FSessionTraceFilterService::OnEndFrame()
 		FString DisabledChannels;
 		GenerateConcatenatedChannels(FrameDisabledChannels, DisabledChannels);
 				
-		if (SessionService.IsValid())
-		{
-			UE_LOG(LogTemp, Display, TEXT("CHANNELS %s: %d"), *DisabledChannels, false);
-			SessionService->ToggleChannels(Handle, *DisabledChannels, false);
-		}
+		UE_LOG(LogTemp, Display, TEXT("CHANNELS %s: %d"), *DisabledChannels, false);
+		ControlClient.SendToggleChannel(*DisabledChannels, false);
 
 		FrameDisabledChannels.Empty();
 	}
+
+	ControlClient.Disconnect();
 }
 
 void FSessionTraceFilterService::DisableAllChannels()
diff --git a/Engine/Plugins/Developer/TraceDataFiltering/Source/TraceDataFiltering/TraceDataFiltering.build.cs b/Engine/Plugins/Developer/TraceDataFiltering/Source/TraceDataFiltering/TraceDataFiltering.build.cs
--- a/Engine/Plugins/Developer/TraceDataFiltering/Source/TraceDataFiltering/TraceDataFiltering.build.cs
+++ b/Engine/Plugins/Developer/TraceDataFiltering/Source/TraceDataFiltering/TraceDataFiltering.build.cs
@@ -13,11 +13,12 @@ public class TraceDataFiltering : ModuleRules
                     "SlateCore",
                     "Slate",
                     "InputCore",
-					"EditorStyle",                    
+					"EditorStyle",
                     "TraceLog",
 					"TraceAnalysis",
 					"TraceServices",
                     "TraceInsights",
+					"Sockets"
                 }
             );
 
