CIS fix due to CL split

diff --git a/Engine/Plugins/Developer/TraceDataFiltering/Source/TraceDataFiltering/Private/TraceDataFilteringModule.cpp b/Engine/Plugins/Developer/TraceDataFiltering/Source/TraceDataFiltering/Private/TraceDataFilteringModule.cpp
--- a/Engine/Plugins/Developer/TraceDataFiltering/Source/TraceDataFiltering/Private/TraceDataFilteringModule.cpp
+++ b/Engine/Plugins/Developer/TraceDataFiltering/Source/TraceDataFiltering/Private/TraceDataFilteringModule.cpp
@@ -21,7 +21,6 @@
 
 #include "STraceDataFilterWidget.h"
 #include "EventFilterStyle.h"
-#include "EventFilterService.h"
 
 IMPLEMENT_MODULE(FTraceFilteringModule, TraceDataFiltering);
 
@@ -35,8 +34,6 @@ void FTraceFilteringModule::StartupModule()
 
 	FConfigCacheIni::LoadGlobalIniFile(TraceFiltersIni, TEXT("TraceDataFilters"));
 
-	FEventFilterService::Get();
-	
 	const FSlateIcon TabIcon(FEventFilterStyle::GetStyleSetName(), "EventFilter.TabIcon");
 	FTabSpawnerEntry& FilterTabSpawnerEntry = FGlobalTabmanager::Get()->RegisterNomadTabSpawner(FTraceFilteringModule::InsightsFilterTabName,
 		FOnSpawnTab::CreateLambda([](const FSpawnTabArgs& Args)
@@ -65,4 +62,4 @@ void FTraceFilteringModule::ShutdownModule()
 	FEventFilterStyle::Shutdown();
 
 	FGlobalTabmanager::Get()->UnregisterNomadTabSpawner(FTraceFilteringModule::InsightsFilterTabName);
-}
\ No newline at end of file
+}
