Insights channel filtering
- Enabling TraceDataFilters for UnrealInsight standalone tool
- Extend TimingProfiler tab to allow spawning Trace Data Filtering tab

diff --git a/Engine/Plugins/Developer/TraceDataFiltering/Source/TraceDataFiltering/Private/TraceDataFilteringModule.cpp b/Engine/Plugins/Developer/TraceDataFiltering/Source/TraceDataFiltering/Private/TraceDataFilteringModule.cpp
--- a/Engine/Plugins/Developer/TraceDataFiltering/Source/TraceDataFiltering/Private/TraceDataFilteringModule.cpp
+++ b/Engine/Plugins/Developer/TraceDataFiltering/Source/TraceDataFiltering/Private/TraceDataFilteringModule.cpp
@@ -34,6 +34,7 @@ void FTraceFilteringModule::StartupModule()
 
 	FConfigCacheIni::LoadGlobalIniFile(TraceFiltersIni, TEXT("TraceDataFilters"));
 
+#if WITH_EDITOR
 	const FSlateIcon TabIcon(FEventFilterStyle::GetStyleSetName(), "EventFilter.TabIcon");
 	FTabSpawnerEntry& FilterTabSpawnerEntry = FGlobalTabmanager::Get()->RegisterNomadTabSpawner(FTraceFilteringModule::InsightsFilterTabName,
 		FOnSpawnTab::CreateLambda([](const FSpawnTabArgs& Args)
@@ -50,16 +51,47 @@ void FTraceFilteringModule::StartupModule()
 	.SetIcon(TabIcon)
 	.SetTooltipText(NSLOCTEXT("FTraceInsightsModule", "FilteringTabTooltip", "Opens the Trace Data Filtering tab, allows for setting Trace Channel states"));
 
-#if WITH_EDITOR
 	FilterTabSpawnerEntry.SetGroup(WorkspaceMenu::GetMenuStructure().GetDeveloperToolsProfilingCategory());
 #else
-	FilterTabSpawnerEntry.SetGroup(WorkspaceMenu::GetMenuStructure().GetToolsCategory());
-#endif	
+	IUnrealInsightsModule& UnrealInsightsModule = FModuleManager::LoadModuleChecked<IUnrealInsightsModule>("TraceInsights");
+
+	FOnRegisterMajorTabExtensions& TimingProfilerLayoutExtension = UnrealInsightsModule.OnRegisterMajorTabExtension(FInsightsManagerTabs::TimingProfilerTabId);
+	TimingProfilerLayoutExtension.AddRaw(this, &FTraceFilteringModule::RegisterTimingProfilerLayoutExtensions);
+#endif
 }
 
 void FTraceFilteringModule::ShutdownModule()
 {
 	FEventFilterStyle::Shutdown();
-
+#if WITH_EDITOR
 	FGlobalTabmanager::Get()->UnregisterNomadTabSpawner(FTraceFilteringModule::InsightsFilterTabName);
+#else
+	IUnrealInsightsModule& UnrealInsightsModule = FModuleManager::LoadModuleChecked<IUnrealInsightsModule>("TraceInsights");
+	FOnRegisterMajorTabExtensions& TimingProfilerLayoutExtension = UnrealInsightsModule.OnRegisterMajorTabExtension(FInsightsManagerTabs::TimingProfilerTabId);
+	TimingProfilerLayoutExtension.RemoveAll(this);
+#endif
 }
+
+void FTraceFilteringModule::RegisterTimingProfilerLayoutExtensions(FInsightsMajorTabExtender& InOutExtender)
+{
+	TSharedRef<FWorkspaceItem> Category = InOutExtender.GetTabManager()->AddLocalWorkspaceMenuCategory(NSLOCTEXT("FTraceInsightsModule", "FilteringCategoryLabel", "Filtering"));
+
+	FInsightsMinorTabConfig& MinorTabConfig = InOutExtender.AddMinorTabConfig();
+	MinorTabConfig.TabId = FTraceFilteringModule::InsightsFilterTabName;
+	MinorTabConfig.TabLabel = NSLOCTEXT("FTraceInsightsModule", "FilteringTabTitle", "Trace Data Filtering");
+	MinorTabConfig.TabTooltip = NSLOCTEXT("FTraceInsightsModule", "FilteringTabTooltip", "Opens the Trace Data Filtering tab, allows for setting Trace Channel states");
+	MinorTabConfig.TabIcon = FSlateIcon(FEventFilterStyle::GetStyleSetName(), "EventFilter.TabIcon");		
+	MinorTabConfig.WorkspaceGroup = Category;
+	MinorTabConfig.OnSpawnTab = FOnSpawnTab::CreateLambda([](const FSpawnTabArgs& Args)
+	{
+		const TSharedRef<SDockTab> DockTab = SNew(SDockTab)
+			.TabRole(ETabRole::NomadTab);
+
+		TSharedRef<STraceDataFilterWidget> Window = SNew(STraceDataFilterWidget);
+		DockTab->SetContent(Window);
+
+		return DockTab;
+	});
+
+	InOutExtender.GetLayoutExtender().ExtendLayout(FTimingProfilerTabs::TimersID, ELayoutExtensionPosition::Before, FTabManager::FTab(FTraceFilteringModule::InsightsFilterTabName, ETabState::ClosedTab));
+}
\ No newline at end of file
diff --git a/Engine/Plugins/Developer/TraceDataFiltering/Source/TraceDataFiltering/Private/TraceDataFilteringModule.h b/Engine/Plugins/Developer/TraceDataFiltering/Source/TraceDataFiltering/Private/TraceDataFilteringModule.h
--- a/Engine/Plugins/Developer/TraceDataFiltering/Source/TraceDataFiltering/Private/TraceDataFilteringModule.h
+++ b/Engine/Plugins/Developer/TraceDataFiltering/Source/TraceDataFiltering/Private/TraceDataFilteringModule.h
@@ -5,6 +5,8 @@
 #include "CoreMinimal.h"
 #include "Modules/ModuleInterface.h"
 
+struct FInsightsMajorTabExtender;
+
 class FTraceFilteringModule : public IModuleInterface
 {
 public:
@@ -14,5 +16,6 @@ public:
 
 	static FString TraceFiltersIni;
 protected:
+	void RegisterTimingProfilerLayoutExtensions(FInsightsMajorTabExtender& InOutExtender);
 	static FName InsightsFilterTabName;	
 };
diff --git a/Engine/Plugins/Developer/TraceDataFiltering/TraceDataFilters.uplugin b/Engine/Plugins/Developer/TraceDataFiltering/TraceDataFilters.uplugin
--- a/Engine/Plugins/Developer/TraceDataFiltering/TraceDataFilters.uplugin
+++ b/Engine/Plugins/Developer/TraceDataFiltering/TraceDataFilters.uplugin
@@ -1,24 +1,28 @@
 {
-	"FileVersion" : 1,
-	"Version" : 1,
-	"VersionName" : "1.0",
-	"FriendlyName" : "Trace Data Filtering",
-	"Description" : "Allows for turning on/off individual or sets of Trace Channels.",
-	"Category" : "Editor",
-	"CreatedBy" : "Epic Games, Inc.",
-	"CreatedByURL" : "http://epicgames.com",
-	"DocsURL" : "",
-	"MarketplaceURL" : "",
-	"SupportURL" : "",
-	"EnabledByDefault" : false,
-	"CanContainContent" : false,
-	"IsBetaVersion" : false,
-	"Installed" : false,
+	"FileVersion": 1,
+	"Version": 1,
+	"VersionName": "1.0",
+	"FriendlyName": "Trace Data Filtering",
+	"Description": "Allows for turning on/off individual or sets of Trace Channels.",
+	"Category": "Editor",
+	"CreatedBy": "Epic Games, Inc.",
+	"CreatedByURL": "http://epicgames.com",
+	"DocsURL": "",
+	"MarketplaceURL": "",
+	"SupportURL": "",
+	"EnabledByDefault": false,
+	"CanContainContent": false,
+	"IsBetaVersion": false,
+	"Installed": false,
+	"SupportedPrograms": [ "UnrealInsights" ],
 	"Modules": [
 		{
 			"Name": "TraceDataFiltering",
-			"Type": "Editor",
-			"LoadingPhase": "PreDefault"
+			"Type": "EditorAndProgram",
+			"LoadingPhase": "PreDefault",
+			"WhitelistPrograms": [
+				"UnrealInsights"
+			]
 		}
 	]
 }
\ No newline at end of file
diff --git a/Engine/Programs/UnrealInsights/Config/DefaultEngine.ini b/Engine/Programs/UnrealInsights/Config/DefaultEngine.ini
--- a/Engine/Programs/UnrealInsights/Config/DefaultEngine.ini
+++ b/Engine/Programs/UnrealInsights/Config/DefaultEngine.ini
@@ -1,4 +1,3 @@
 [Plugins]
 +ProgramEnabledPlugins=GameplayInsights
-
-
++ProgramEnabledPlugins=TraceDataFilters
\ No newline at end of file
