Trace Data Filtering infrequently fails to detect active trace channels from Editor -trace=<options> commandline
#fix Respect initially enabled channels (through commandline on running session), and leave them enabled when applying preset(s) for the first time

diff --git a/Engine/Plugins/Developer/TraceDataFiltering/Source/TraceDataFiltering/Private/Trace/SessionTraceFilterService.cpp b/Engine/Plugins/Developer/TraceDataFiltering/Source/TraceDataFiltering/Private/Trace/SessionTraceFilterService.cpp
--- a/Engine/Plugins/Developer/TraceDataFiltering/Source/TraceDataFiltering/Private/Trace/SessionTraceFilterService.cpp
+++ b/Engine/Plugins/Developer/TraceDataFiltering/Source/TraceDataFiltering/Private/Trace/SessionTraceFilterService.cpp
@@ -17,6 +17,8 @@
 FSessionTraceFilterService::FSessionTraceFilterService(Trace::FSessionHandle InHandle, TSharedPtr<const Trace::IAnalysisSession> InSession) : Session(InSession), Handle(InHandle)
 {
 	FCoreDelegates::OnEndFrame.AddRaw(this, &FSessionTraceFilterService::OnEndFrame);
+	// Retrieve the channels currently enabled on the provider, this'll be the ones specified on the commandline (-trace=ChannelX)
+	RetrieveAndStoreStartupChannels();
 }
 
 void FSessionTraceFilterService::GetRootObjects(TArray<FTraceObjectInfo>& OutObjects) const
@@ -86,6 +88,13 @@ void FSessionTraceFilterService::UpdateFilterPresets(const TArray<TSharedPtr<IFi
 	{
 		SetObjectFilterState(EventName, true);
 	}
+
+	for (const FString& EventName : FrameZeroEnabledChannels)
+	{
+		SetObjectFilterState(EventName, true);
+	}
+	// Commandline channels are only applied once when changing presets
+	FrameZeroEnabledChannels.Empty();
 }
 
 void FSessionTraceFilterService::OnEndFrame()
@@ -157,10 +166,27 @@ void FSessionTraceFilterService::DisableAllChannels()
 	if (ChannelProvider)
 	{
 		const uint64 ChannelCount = ChannelProvider->GetChannelCount();
-		auto Channels = ChannelProvider->GetChannels();
+		const TArray<Trace::FChannelEntry>& Channels = ChannelProvider->GetChannels();
 		for (uint64 ChannelIndex = 0; ChannelIndex < Channels.Num(); ++ChannelIndex)
 		{
 			SetObjectFilterState(Channels[ChannelIndex].Name, false);
 		}
 	}
 }
+
+void FSessionTraceFilterService::RetrieveAndStoreStartupChannels()
+{
+	const Trace::IChannelProvider* ChannelProvider = Session->ReadProvider<Trace::IChannelProvider>("ChannelProvider");	
+	if (ChannelProvider)
+	{
+		const uint64 ChannelCount = ChannelProvider->GetChannelCount();
+		const TArray<Trace::FChannelEntry>& Channels = ChannelProvider->GetChannels();
+		for (uint64 ChannelIndex = 0; ChannelIndex < Channels.Num(); ++ChannelIndex)
+		{
+			if (Channels[ChannelIndex].bIsEnabled)
+			{
+				FrameZeroEnabledChannels.Add(Channels[ChannelIndex].Name);
+			}
+		}
+	}
+}
diff --git a/Engine/Plugins/Developer/TraceDataFiltering/Source/TraceDataFiltering/Private/Trace/SessionTraceFilterService.h b/Engine/Plugins/Developer/TraceDataFiltering/Source/TraceDataFiltering/Private/Trace/SessionTraceFilterService.h
--- a/Engine/Plugins/Developer/TraceDataFiltering/Source/TraceDataFiltering/Private/Trace/SessionTraceFilterService.h
+++ b/Engine/Plugins/Developer/TraceDataFiltering/Source/TraceDataFiltering/Private/Trace/SessionTraceFilterService.h
@@ -35,6 +35,8 @@ protected:
 
 	/** Retrieves channels names from provider and marks them all as disabled */
 	void DisableAllChannels();
+
+	void RetrieveAndStoreStartupChannels();
 protected:
 	/** Session this instance represents the filtering service for */
 	TSharedPtr<const Trace::IAnalysisSession> Session;
@@ -44,7 +46,9 @@ protected:
 	/** Names of channels that were either enabled or disabled during the duration of this frame */
 	TSet<FString> FrameEnabledChannels;
 	TSet<FString> FrameDisabledChannels;
+
+	TSet<FString> FrameZeroEnabledChannels;
 	
 	/** Timestamp at which contained data (including provider) was last updated */
 	FDateTime TimeStamp;	
-};
+};
\ No newline at end of file
