Separate out (engine) filter presets into their own UObject container to mitigated incorrect ini saving

diff --git a/Engine/Plugins/Developer/TraceDataFiltering/Config/BaseTraceDataFilters.ini b/Engine/Plugins/Developer/TraceDataFiltering/Config/BaseTraceDataFilters.ini
--- a/Engine/Plugins/Developer/TraceDataFiltering/Config/BaseTraceDataFilters.ini
+++ b/Engine/Plugins/Developer/TraceDataFiltering/Config/BaseTraceDataFilters.ini
@@ -1,2 +1,2 @@
-[/Script/TraceDataFiltering.SharedFilterPresetContainer]
+[/Script/TraceDataFiltering.EngineFilterPresetContainer]
 +EnginePresets=(Name="Animation",WhitelistedNames=("Animation", "Object", "Frame"))
diff --git a/Engine/Plugins/Developer/TraceDataFiltering/Source/TraceDataFiltering/Private/FilterPresets.cpp b/Engine/Plugins/Developer/TraceDataFiltering/Source/TraceDataFiltering/Private/FilterPresets.cpp
--- a/Engine/Plugins/Developer/TraceDataFiltering/Source/TraceDataFiltering/Private/FilterPresets.cpp
+++ b/Engine/Plugins/Developer/TraceDataFiltering/Source/TraceDataFiltering/Private/FilterPresets.cpp
@@ -32,7 +32,7 @@ void USharedFilterPresetContainer::Save()
 	SharedPresetsContainer->UpdateDefaultConfigFile();
 }
 
-void USharedFilterPresetContainer::GetEnginePresets(TArray<TSharedPtr<IFilterPreset>>& OutPresets)
+void UEngineFilterPresetContainer::GetEnginePresets(TArray<TSharedPtr<IFilterPreset>>& OutPresets)
 {
 	for (FFilterData& FilterData : EnginePresets)
 	{
diff --git a/Engine/Plugins/Developer/TraceDataFiltering/Source/TraceDataFiltering/Private/FilterPresets.h b/Engine/Plugins/Developer/TraceDataFiltering/Source/TraceDataFiltering/Private/FilterPresets.h
--- a/Engine/Plugins/Developer/TraceDataFiltering/Source/TraceDataFiltering/Private/FilterPresets.h
+++ b/Engine/Plugins/Developer/TraceDataFiltering/Source/TraceDataFiltering/Private/FilterPresets.h
@@ -110,7 +110,6 @@ class USharedFilterPresetContainer : public UObject
 
 	friend struct FFilterPresetHelpers;
 public:
-	void GetEnginePresets(TArray<TSharedPtr<IFilterPreset>>& OutPresets);
 	void GetSharedUserPresets(TArray<TSharedPtr<IFilterPreset>>& OutPresets);
 
 	static void AddFilterData(const FFilterData& InFilterData);
@@ -119,7 +118,17 @@ public:
 protected:
 	UPROPERTY(Config)
 	TArray<FFilterData> SharedPresets;
+};
+
+UCLASS(Config = TraceDataFilters, DefaultConfig)
+class UEngineFilterPresetContainer : public UObject
+{
+	GENERATED_BODY()
 
+	friend struct FFilterPresetHelpers;
+public:
+	void GetEnginePresets(TArray<TSharedPtr<IFilterPreset>>& OutPresets);
+protected:
 	UPROPERTY(Config)
 	TArray<FFilterData> EnginePresets;
 };
\ No newline at end of file
diff --git a/Engine/Plugins/Developer/TraceDataFiltering/Source/TraceDataFiltering/Private/SFilterPresetList.cpp b/Engine/Plugins/Developer/TraceDataFiltering/Source/TraceDataFiltering/Private/SFilterPresetList.cpp
--- a/Engine/Plugins/Developer/TraceDataFiltering/Source/TraceDataFiltering/Private/SFilterPresetList.cpp
+++ b/Engine/Plugins/Developer/TraceDataFiltering/Source/TraceDataFiltering/Private/SFilterPresetList.cpp
@@ -45,6 +45,7 @@ void SFilterPresetList::Construct( const FArguments& InArgs )
 	
 	PresetContainer = GetMutableDefault<ULocalFilterPresetContainer>();
 	SharedPresetContainer = GetMutableDefault<USharedFilterPresetContainer>();
+	EnginePresetContainer = GetMutableDefault<UEngineFilterPresetContainer>();
 	
 	RefreshFilterPresets();
 
@@ -89,7 +90,7 @@ void SFilterPresetList::RefreshFilterPresets()
 
 	PresetContainer->GetUserPresets(UserFilterPresets);
 
-	SharedPresetContainer->GetEnginePresets(EngineFilterPresets);
+	EnginePresetContainer->GetEnginePresets(EngineFilterPresets);
 	SharedPresetContainer->GetSharedUserPresets(SharedUserFilterPresets);
 
 	AllFilterPresets.Append(EngineFilterPresets);
diff --git a/Engine/Plugins/Developer/TraceDataFiltering/Source/TraceDataFiltering/Private/SFilterPresetList.h b/Engine/Plugins/Developer/TraceDataFiltering/Source/TraceDataFiltering/Private/SFilterPresetList.h
--- a/Engine/Plugins/Developer/TraceDataFiltering/Source/TraceDataFiltering/Private/SFilterPresetList.h
+++ b/Engine/Plugins/Developer/TraceDataFiltering/Source/TraceDataFiltering/Private/SFilterPresetList.h
@@ -16,6 +16,7 @@ class SFilterPreset;
 class SWrapBox;
 class ULocalFilterPresetContainer;
 class USharedFilterPresetContainer;
+class UEngineFilterPresetContainer;
 
 /**
  * A list of presets currently applied to insights filtering window
@@ -114,6 +115,7 @@ private:
 	/** Weak pointers to filter preset containers */
 	TWeakObjectPtr<ULocalFilterPresetContainer> PresetContainer;
 	TWeakObjectPtr<USharedFilterPresetContainer> SharedPresetContainer;
+	TWeakObjectPtr<UEngineFilterPresetContainer> EnginePresetContainer;
 
 	/** Delegate for when Presets have changed */
 	FOnPresetsChanged OnPresetsChanged;
