Trace Data Filtering window is blank after opening a non-live trace
#fix Added warning widget, visible when there is no valid live session available
#misc fixed up widget style for standalone Insights

diff --git a/Engine/Plugins/Developer/TraceDataFiltering/Source/TraceDataFiltering/Private/EventFilterStyle.cpp b/Engine/Plugins/Developer/TraceDataFiltering/Source/TraceDataFiltering/Private/EventFilterStyle.cpp
--- a/Engine/Plugins/Developer/TraceDataFiltering/Source/TraceDataFiltering/Private/EventFilterStyle.cpp
+++ b/Engine/Plugins/Developer/TraceDataFiltering/Source/TraceDataFiltering/Private/EventFilterStyle.cpp
@@ -7,10 +7,13 @@
 
 TSharedPtr< FSlateStyleSet > FEventFilterStyle::StyleSet = nullptr;
 
+FTextBlockStyle FEventFilterStyle::NormalText;
+
 #define IMAGE_BRUSH( RelativePath, ... ) FSlateImageBrush( StyleSet->RootToContentDir( RelativePath, TEXT(".png") ), __VA_ARGS__ )
 #define BOX_BRUSH( RelativePath, ... ) FSlateBoxBrush( StyleSet->RootToContentDir( RelativePath, TEXT(".png") ), __VA_ARGS__ )
 #define BORDER_BRUSH( RelativePath, ... ) FSlateBorderBrush( StyleSet->RootToContentDir( RelativePath, TEXT(".png") ), __VA_ARGS__ )
 #define DEFAULT_FONT(...) FCoreStyle::GetDefaultFontStyle(__VA_ARGS__)
+#define ICON_FONT(...) FSlateFontInfo(StyleSet->RootToContentDir("Fonts/FontAwesome", TEXT(".ttf")), __VA_ARGS__)
 
 // Const icon sizes
 static const FVector2D Icon8x8(8.0f, 8.0f);
@@ -45,6 +48,16 @@ void FEventFilterStyle::Initialize()
 	StyleSet->SetContentRoot(FPaths::EngineContentDir() / TEXT("Editor/Slate"));
 	StyleSet->SetCoreContentRoot(FPaths::EngineContentDir() / TEXT("Slate"));
 
+	StyleSet->Set("EventFilter.GroupBorder", new BOX_BRUSH("Common/GroupBorder", FMargin(4.0f / 16.0f)));
+
+	NormalText = FTextBlockStyle()
+		.SetFont(DEFAULT_FONT("Regular", FCoreStyle::RegularTextSize))
+		.SetColorAndOpacity(FSlateColor::UseForeground())
+		.SetShadowOffset(FVector2D::ZeroVector)
+		.SetShadowColorAndOpacity(FLinearColor::Black)
+		.SetHighlightColor(FLinearColor(0.02f, 0.3f, 0.0f))
+		.SetHighlightShape(BOX_BRUSH("Common/TextBlockHighlightShape", FMargin(3.f / 8.f)));
+
 	// Colors
 	{
 		StyleSet->Set("EventFilter.EnginePreset", FLinearColor(0.728f, 0.364f, 0.003f));
@@ -66,6 +79,32 @@ void FEventFilterStyle::Initialize()
 		StyleSet->Set("EventFilter.TabIcon", new IMAGE_BRUSH("/Icons/icon_Genericfinder_16x", Icon16x16));
 	}
 
+	FButtonStyle Button = FButtonStyle()
+		.SetNormal(BOX_BRUSH("Common/Button", FVector2D(32, 32), 8.0f / 32.0f))
+		.SetHovered(BOX_BRUSH("Common/Button_Hovered", FVector2D(32, 32), 8.0f / 32.0f))
+		.SetPressed(BOX_BRUSH("Common/Button_Pressed", FVector2D(32, 32), 8.0f / 32.0f))
+		.SetNormalPadding(FMargin(2, 2, 2, 2))
+		.SetPressedPadding(FMargin(2, 3, 2, 1));
+
+	StyleSet->Set("ToggleButton", FButtonStyle(Button)
+		.SetNormal(FSlateNoResource())
+		.SetHovered(BOX_BRUSH("Common/RoundedSelection_16x", 4.0f / 16.0f, FLinearColor(0.701f, 0.225f, 0.003f)))
+		.SetPressed(BOX_BRUSH("Common/RoundedSelection_16x", 4.0f / 16.0f, FLinearColor(0.701f, 0.225f, 0.003f)))
+	);
+
+	FComboButtonStyle ToolbarComboButton = FComboButtonStyle()
+		.SetButtonStyle(StyleSet->GetWidgetStyle<FButtonStyle>("ToggleButton"))
+		.SetDownArrowImage(IMAGE_BRUSH("Common/ShadowComboArrow", Icon8x8))
+		.SetMenuBorderBrush(BOX_BRUSH("Old/Menu_Background", FMargin(8.0f / 64.0f)))
+		.SetMenuBorderPadding(FMargin(0.0f));
+	StyleSet->Set("EventFilter.ComboButton", ToolbarComboButton);
+	
+	StyleSet->Set("EventFilter.TextStyle", FTextBlockStyle(NormalText)
+		.SetFont(DEFAULT_FONT("Bold", 9))
+		.SetColorAndOpacity(FLinearColor(1.0f, 1.0f, 1.0f, 0.9f))
+		.SetShadowOffset(FVector2D(1, 1))
+		.SetShadowColorAndOpacity(FLinearColor(0, 0, 0, 0.9f)));
+	
 
 	// Filter list
 	/* Set images for various SCheckBox states associated with "ContentBrowser.FilterButton" ... */
@@ -81,6 +120,11 @@ void FEventFilterStyle::Initialize()
 
 	StyleSet->Set("FilterPresets.FilterNameFont", DEFAULT_FONT("Regular", 10));
 	StyleSet->Set("FilterPresets.FilterButtonBorder", new BOX_BRUSH("Common/RoundedSelection_16x", FMargin(4.0f / 16.0f)));
+
+	StyleSet->Set("FilterPresets.SessionWarningBorder", new BOX_BRUSH("Common/GroupBorderLight", FMargin(4.0f / 16.0f)));
+	StyleSet->Set("FilterPresets.WarningIcon", new IMAGE_BRUSH("Settings/Settings_Warning", Icon40x40));
+
+	StyleSet->Set("FontAwesome.9", ICON_FONT(9));
 	
 	FSlateStyleRegistry::RegisterSlateStyle( *StyleSet.Get() );
 }
@@ -89,6 +133,7 @@ void FEventFilterStyle::Initialize()
 #undef BOX_BRUSH
 #undef BORDER_BRUSH
 #undef DEFAULT_FONT
+#undef ICON_FONT
 
 void FEventFilterStyle::Shutdown()
 {
diff --git a/Engine/Plugins/Developer/TraceDataFiltering/Source/TraceDataFiltering/Private/EventFilterStyle.h b/Engine/Plugins/Developer/TraceDataFiltering/Source/TraceDataFiltering/Private/EventFilterStyle.h
--- a/Engine/Plugins/Developer/TraceDataFiltering/Source/TraceDataFiltering/Private/EventFilterStyle.h
+++ b/Engine/Plugins/Developer/TraceDataFiltering/Source/TraceDataFiltering/Private/EventFilterStyle.h
@@ -22,4 +22,5 @@ private:
 
 	/** Singleton instances of this style. */
 	static TSharedPtr< class FSlateStyleSet > StyleSet;	
+	static FTextBlockStyle NormalText;
 };
diff --git a/Engine/Plugins/Developer/TraceDataFiltering/Source/TraceDataFiltering/Private/STraceDataFilterWidget.cpp b/Engine/Plugins/Developer/TraceDataFiltering/Source/TraceDataFiltering/Private/STraceDataFilterWidget.cpp
--- a/Engine/Plugins/Developer/TraceDataFiltering/Source/TraceDataFiltering/Private/STraceDataFilterWidget.cpp
+++ b/Engine/Plugins/Developer/TraceDataFiltering/Source/TraceDataFiltering/Private/STraceDataFilterWidget.cpp
@@ -21,6 +21,7 @@
 #include "Insights/IUnrealInsightsModule.h"
 #include "Trace/StoreClient.h"
 #include "Trace/SessionTraceFilterService.h"
+#include "EventFilterStyle.h"
 
 #define LOCTEXT_NAMESPACE "STraceDataFilterWidget"
 
@@ -51,7 +52,7 @@ void STraceDataFilterWidget::Construct(const FArguments& InArgs)
 	[		
 		SNew(SBorder)
 		.Padding(4)
-		.BorderImage(FEditorStyle::GetBrush("ToolPanel.GroupBorder"))
+		.BorderImage(FEventFilterStyle::GetBrush("EventFilter.GroupBorder"))
 		[
 			SNew( SVerticalBox )
 
@@ -68,7 +69,7 @@ void STraceDataFilterWidget::Construct(const FArguments& InArgs)
 				[
 					SNew(SComboButton)
 					.Visibility(EVisibility::Visible)
-					.ComboButtonStyle(FEditorStyle::Get(), "GenericFilters.ComboButtonStyle")
+					.ComboButtonStyle(FEventFilterStyle::Get(), "EventFilter.ComboButton")
 					.ForegroundColor(FLinearColor::White)
 					.ContentPadding(0.0f)
 					.OnGetMenuContent(this, &STraceDataFilterWidget::MakeAddFilterMenu)
@@ -80,8 +81,8 @@ void STraceDataFilterWidget::Construct(const FArguments& InArgs)
 						.VAlign(VAlign_Center)
 						[
 							SNew(STextBlock)
-							.TextStyle(FEditorStyle::Get(), "GenericFilters.TextStyle")
-							.Font(FEditorStyle::Get().GetFontStyle("FontAwesome.9"))
+							.TextStyle(FEventFilterStyle::Get(), "EventFilter.TextStyle")
+							.Font(FEventFilterStyle::Get().GetFontStyle("FontAwesome.9"))
 							.Text(FText::FromString(FString(TEXT("\xf0b0"))) /*fa-filter*/)
 						]
 						+SHorizontalBox::Slot()
@@ -90,7 +91,7 @@ void STraceDataFilterWidget::Construct(const FArguments& InArgs)
 						.VAlign(VAlign_Center)
 						[
 							SNew( STextBlock )
-							.TextStyle(FEditorStyle::Get(), "GenericFilters.TextStyle")
+							.TextStyle(FEventFilterStyle::Get(), "EventFilter.TextStyle")
 							.Text( LOCTEXT("PresetsMenuLabel", "Filter Presets") )
 						]
 					]
@@ -121,6 +122,39 @@ void STraceDataFilterWidget::Construct(const FArguments& InArgs)
 				]
 			]
 
+			+ SVerticalBox::Slot()
+			.AutoHeight()
+			[
+				SNew(SBorder)
+				.BorderImage(FEventFilterStyle::GetBrush("FilterPresets.SessionWarningBorder"))
+				.BorderBackgroundColor(FColor(166,137,0))
+				[
+					SNew(SHorizontalBox)
+					.Visibility_Lambda([this]() -> EVisibility 
+					{
+						return HasValidFilterSession() ? EVisibility::Collapsed : EVisibility::Visible;
+					})
+
+					+ SHorizontalBox::Slot()
+					.VAlign(VAlign_Center)
+					.AutoWidth()
+					.Padding(4.0f, 0.0f, 4.0f, 0.0f)
+					[
+						SNew(SImage)
+						.Image(FEventFilterStyle::GetBrush("FilterPresets.WarningIcon"))
+					]
+
+					+SHorizontalBox::Slot()
+					.VAlign(VAlign_Center)
+					.AutoWidth()
+					.Padding(4.0f, 0.0f, 0.0f, 0.0f)
+					[
+						SNew(STextBlock)
+						.Text(LOCTEXT("NoValidFilterWarning", "Trace Data Filtering requires a connected Live Trace Session"))
+					]
+				]
+			]
+
 			+ SVerticalBox::Slot()	
 			.Padding(0.0f, 0.0f, 0.0f, 2.0f)
 			[
diff --git a/Engine/Plugins/Developer/TraceSourceFiltering/Source/SourceFilteringEditor/Private/STraceSourceFilteringWidget.cpp b/Engine/Plugins/Developer/TraceSourceFiltering/Source/SourceFilteringEditor/Private/STraceSourceFilteringWidget.cpp
--- a/Engine/Plugins/Developer/TraceSourceFiltering/Source/SourceFilteringEditor/Private/STraceSourceFilteringWidget.cpp
+++ b/Engine/Plugins/Developer/TraceSourceFiltering/Source/SourceFilteringEditor/Private/STraceSourceFilteringWidget.cpp
@@ -64,7 +64,7 @@ void STraceSourceFilteringWidget::Construct(const FArguments& InArgs)
 	[
 		SNew(SBorder)
 		.Padding(4.0f)
-		.BorderImage(FEditorStyle::GetBrush("ToolPanel.GroupBorder"))
+		.BorderImage(FSourceFilterStyle::GetBrush("SourceFilter.GroupBorder"))
 		[
 			SNew(SVerticalBox)
 			+ SVerticalBox::Slot()
@@ -186,7 +186,7 @@ void STraceSourceFilteringWidget::ConstructMenuBox()
 			[
 				SNew(STextBlock)
 				.TextStyle(FSourceFilterStyle::Get(), "SourceFilter.TextStyle")
-				.Font(FEditorStyle::Get().GetFontStyle("FontAwesome.9"))
+				.Font(FSourceFilterStyle::Get().GetFontStyle("FontAwesome.9"))
 				.Text(FText::FromString(FString(TEXT("\xf0fe"))) /*fa-filter*/)
 			]
 			+ SHorizontalBox::Slot()
@@ -221,7 +221,7 @@ void STraceSourceFilteringWidget::ConstructMenuBox()
 			[
 				SNew(STextBlock)
 				.TextStyle(FSourceFilterStyle::Get(), "SourceFilter.TextStyle")
-				.Font(FEditorStyle::Get().GetFontStyle("FontAwesome.9"))
+				.Font(FSourceFilterStyle::Get().GetFontStyle("FontAwesome.9"))
 				.Text(FText::FromString(FString(TEXT("\xf0b0"))) /*fa-filter*/)
 			]
 			+SHorizontalBox::Slot()
diff --git a/Engine/Plugins/Developer/TraceSourceFiltering/Source/SourceFilteringEditor/Private/SourceFilterStyle.cpp b/Engine/Plugins/Developer/TraceSourceFiltering/Source/SourceFilteringEditor/Private/SourceFilterStyle.cpp
--- a/Engine/Plugins/Developer/TraceSourceFiltering/Source/SourceFilteringEditor/Private/SourceFilterStyle.cpp
+++ b/Engine/Plugins/Developer/TraceSourceFiltering/Source/SourceFilteringEditor/Private/SourceFilterStyle.cpp
@@ -18,6 +18,7 @@ FTextBlockStyle FSourceFilterStyle::NormalText;
 #define BOX_BRUSH( RelativePath, ... ) FSlateBoxBrush( StyleSet->RootToContentDir( RelativePath, TEXT(".png") ), __VA_ARGS__ )
 #define BORDER_BRUSH( RelativePath, ... ) FSlateBorderBrush( StyleSet->RootToContentDir( RelativePath, TEXT(".png") ), __VA_ARGS__ )
 #define DEFAULT_FONT(...) FCoreStyle::GetDefaultFontStyle(__VA_ARGS__)
+#define ICON_FONT(...) FSlateFontInfo(StyleSet->RootToContentDir("Fonts/FontAwesome", TEXT(".ttf")), __VA_ARGS__)
 
 // Const icon sizes
 static const FVector2D Icon8x8(8.0f, 8.0f);
@@ -52,6 +53,8 @@ void FSourceFilterStyle::Initialize()
 	StyleSet->SetContentRoot(FPaths::EngineContentDir() / TEXT("Editor/Slate"));
 	StyleSet->SetCoreContentRoot(FPaths::EngineContentDir() / TEXT("Slate"));
 
+	StyleSet->Set("SourceFilter.GroupBorder", new BOX_BRUSH("Common/GroupBorder", FMargin(4.0f / 16.0f)));
+
 
 	NormalText = FTextBlockStyle()
 		.SetFont(DEFAULT_FONT("Regular", FCoreStyle::RegularTextSize))
@@ -149,6 +152,8 @@ void FSourceFilterStyle::Initialize()
 		.SetShadowColorAndOpacity(FLinearColor(0, 0, 0, 0.9f)));
 
 	StyleSet->Set("SourceFilter.DragDrop.Border", new BOX_BRUSH("Old/Menu_Background", FMargin(8.0f / 64.0f)));
+
+	StyleSet->Set("FontAwesome.9", ICON_FONT(9));
 		
 	FSlateStyleRegistry::RegisterSlateStyle( *StyleSet.Get() );
 }
@@ -157,6 +162,7 @@ void FSourceFilterStyle::Initialize()
 #undef BOX_BRUSH
 #undef BORDER_BRUSH
 #undef DEFAULT_FONT
+#undef ICON_FONT
 
 void FSourceFilterStyle::Shutdown()
 {
