Non-pch build fix

diff --git a/Engine/Plugins/Developer/TraceDataFiltering/Source/TraceDataFiltering/Private/EventFilterStyle.h b/Engine/Plugins/Developer/TraceDataFiltering/Source/TraceDataFiltering/Private/EventFilterStyle.h
--- a/Engine/Plugins/Developer/TraceDataFiltering/Source/TraceDataFiltering/Private/EventFilterStyle.h
+++ b/Engine/Plugins/Developer/TraceDataFiltering/Source/TraceDataFiltering/Private/EventFilterStyle.h
@@ -4,6 +4,7 @@
 
 #include "CoreMinimal.h"
 #include "Styling/SlateStyle.h"
+#include "Styling/SlateTypes.h"
 
 class FEventFilterStyle
 {
diff --git a/Engine/Plugins/Developer/TraceDataFiltering/Source/TraceDataFiltering/Private/STraceDataFilterWidget.cpp b/Engine/Plugins/Developer/TraceDataFiltering/Source/TraceDataFiltering/Private/STraceDataFilterWidget.cpp
--- a/Engine/Plugins/Developer/TraceDataFiltering/Source/TraceDataFiltering/Private/STraceDataFilterWidget.cpp
+++ b/Engine/Plugins/Developer/TraceDataFiltering/Source/TraceDataFiltering/Private/STraceDataFilterWidget.cpp
@@ -7,6 +7,7 @@
 #include "Widgets/Input/SComboButton.h"
 #include "Widgets/Input/SSearchBox.h"
 #include "Widgets/SOverlay.h"
+#include "Widgets/Images/SImage.h"
 #include "Templates/SharedPointer.h"
 #include "TraceServices/ITraceServicesModule.h"
 #include "Framework/MultiBox/MultiBoxBuilder.h"
