Fixed use-after-free when shutting down UE4Editor.

#rnx

diff --git a/Engine/Plugins/Developer/TraceDataFiltering/Source/TraceDataFiltering/Private/Trace/SessionTraceFilterService.cpp b/Engine/Plugins/Developer/TraceDataFiltering/Source/TraceDataFiltering/Private/Trace/SessionTraceFilterService.cpp
--- a/Engine/Plugins/Developer/TraceDataFiltering/Source/TraceDataFiltering/Private/Trace/SessionTraceFilterService.cpp
+++ b/Engine/Plugins/Developer/TraceDataFiltering/Source/TraceDataFiltering/Private/Trace/SessionTraceFilterService.cpp
@@ -20,6 +20,11 @@ FSessionTraceFilterService::FSessionTraceFilterService(Trace::FSessionHandle InH
 	RetrieveAndStoreStartupChannels();
 }
 
+FSessionTraceFilterService::~FSessionTraceFilterService()
+{
+	FCoreDelegates::OnEndFrame.RemoveAll(this);
+}
+
 void FSessionTraceFilterService::GetRootObjects(TArray<FTraceObjectInfo>& OutObjects) const
 {
 	const Trace::IChannelProvider* ChannelProvider = Session->ReadProvider<Trace::IChannelProvider>("ChannelProvider");
diff --git a/Engine/Plugins/Developer/TraceDataFiltering/Source/TraceDataFiltering/Private/Trace/SessionTraceFilterService.h b/Engine/Plugins/Developer/TraceDataFiltering/Source/TraceDataFiltering/Private/Trace/SessionTraceFilterService.h
--- a/Engine/Plugins/Developer/TraceDataFiltering/Source/TraceDataFiltering/Private/Trace/SessionTraceFilterService.h
+++ b/Engine/Plugins/Developer/TraceDataFiltering/Source/TraceDataFiltering/Private/Trace/SessionTraceFilterService.h
@@ -19,7 +19,7 @@ class FSessionTraceFilterService : public ISessionTraceFilterService
 {
 public:
 	FSessionTraceFilterService(Trace::FSessionHandle InHandle, TSharedPtr<const Trace::IAnalysisSession> InSession);
-	virtual ~FSessionTraceFilterService() {}
+	virtual ~FSessionTraceFilterService();
 
 	/** Begin ISessionTraceFilterService overrides */
 	virtual void GetRootObjects(TArray<FTraceObjectInfo>& OutObjects) const override;
