Opening the Trace Data Filtering panel causes an extreme drop in editor responsiveness
#fix Split trace data filter service into trace and editor versions, this removes possible network socket contention (used for communicating with trace) in the Editor

diff --git a/Engine/Plugins/Developer/TraceDataFiltering/Source/TraceDataFiltering/Private/BaseSessionFilterService.cpp b/Engine/Plugins/Developer/TraceDataFiltering/Source/TraceDataFiltering/Private/BaseSessionFilterService.cpp
new file mode 100644
--- /dev/null
+++ b/Engine/Plugins/Developer/TraceDataFiltering/Source/TraceDataFiltering/Private/BaseSessionFilterService.cpp
@@ -0,0 +1,133 @@
+// Copyright Epic Games, Inc. All Rights Reserved.
+
+#include "BaseSessionFilterService.h"
+#include "TraceServices/Model/Channel.h"
+#include "TraceServices/ITraceServicesModule.h"
+#include "Templates/SharedPointer.h"
+#include "Misc/CoreDelegates.h"
+#include "Algo/Transform.h"
+#include "Modules/ModuleManager.h"
+#include <Insights/IUnrealInsightsModule.h>
+#include <Trace/StoreClient.h>
+#include <Trace/ControlClient.h>
+#include <IPAddress.h>
+#include <SocketSubsystem.h>
+
+FBaseSessionFilterService::FBaseSessionFilterService(Trace::FSessionHandle InHandle, TSharedPtr<const Trace::IAnalysisSession> InSession) : Session(InSession), Handle(InHandle)
+{
+	FCoreDelegates::OnEndFrame.AddRaw(this, &FBaseSessionFilterService::OnApplyChannelChanges);
+	// Retrieve the channels currently enabled on the provider, this'll be the ones specified on the commandline (-trace=ChannelX)
+	RetrieveAndStoreStartupChannels();
+}
+
+FBaseSessionFilterService::~FBaseSessionFilterService()
+{
+	FCoreDelegates::OnEndFrame.RemoveAll(this);
+}
+
+void FBaseSessionFilterService::GetRootObjects(TArray<FTraceObjectInfo>& OutObjects) const
+{
+	const Trace::IChannelProvider* ChannelProvider = Session->ReadProvider<Trace::IChannelProvider>("ChannelProvider");
+	if (ChannelProvider)
+	{
+		const uint64 ChannelCount = ChannelProvider->GetChannelCount();		
+		const TArray<Trace::FChannelEntry>& Channels = ChannelProvider->GetChannels();
+		for (uint64 ChannelIndex = 0; ChannelIndex < Channels.Num(); ++ChannelIndex)
+		{
+			FTraceObjectInfo& EventInfo = OutObjects.AddDefaulted_GetRef();
+			EventInfo.Name = Channels[ChannelIndex].Name;
+			EventInfo.bEnabled = Channels[ChannelIndex].bIsEnabled;
+			EventInfo.Hash = GetTypeHash(EventInfo.Name);
+			EventInfo.OwnerHash = 0;
+		}
+	}
+}
+
+void FBaseSessionFilterService::GetChildObjects(uint32 InObjectHash, TArray<FTraceObjectInfo>& OutChildObjects) const
+{
+	/** TODO, parent/child relationship for Channels */
+}
+
+const FDateTime& FBaseSessionFilterService::GetTimestamp()
+{
+	const Trace::IChannelProvider* ChannelProvider = Session->ReadProvider<Trace::IChannelProvider>("ChannelProvider");
+	if (ChannelProvider)
+	{
+		TimeStamp = ChannelProvider->GetTimeStamp();
+	}
+
+	return TimeStamp;
+}
+
+void FBaseSessionFilterService::SetObjectFilterState(const FString& InObjectName, const bool bFilterState)
+{
+	if (bFilterState)
+	{
+		FrameDisabledChannels.Remove(InObjectName);
+		FrameEnabledChannels.Add(InObjectName);
+	}
+	else
+	{
+		FrameEnabledChannels.Remove(InObjectName);
+		FrameDisabledChannels.Add(InObjectName);
+	}
+}
+
+void FBaseSessionFilterService::UpdateFilterPresets(const TArray<TSharedPtr<IFilterPreset>>& InPresets)
+{
+	TSet<FString> UniqueNames;
+
+	for (const TSharedPtr<IFilterPreset>& Preset : InPresets)
+	{
+		TArray<FString> Names;
+		Preset->GetWhitelistedNames(Names);
+
+		// We are only interested in the unique channels names, as a result of combining multiple presets
+		Algo::Transform(Names, UniqueNames, [](FString InName) { return InName; });
+	}
+
+	DisableAllChannels();
+
+	for (const FString& EventName : UniqueNames)
+	{
+		SetObjectFilterState(EventName, true);
+	}
+
+	for (const FString& EventName : FrameZeroEnabledChannels)
+	{
+		SetObjectFilterState(EventName, true);
+	}
+	// Commandline channels are only applied once when changing presets
+	FrameZeroEnabledChannels.Empty();
+}
+
+void FBaseSessionFilterService::DisableAllChannels()
+{
+	const Trace::IChannelProvider* ChannelProvider = Session->ReadProvider<Trace::IChannelProvider>("ChannelProvider");
+	if (ChannelProvider)
+	{
+		const uint64 ChannelCount = ChannelProvider->GetChannelCount();
+		const TArray<Trace::FChannelEntry>& Channels = ChannelProvider->GetChannels();
+		for (uint64 ChannelIndex = 0; ChannelIndex < Channels.Num(); ++ChannelIndex)
+		{
+			SetObjectFilterState(Channels[ChannelIndex].Name, false);
+		}
+	}
+}
+
+void FBaseSessionFilterService::RetrieveAndStoreStartupChannels()
+{
+	const Trace::IChannelProvider* ChannelProvider = Session->ReadProvider<Trace::IChannelProvider>("ChannelProvider");	
+	if (ChannelProvider)
+	{
+		const uint64 ChannelCount = ChannelProvider->GetChannelCount();
+		const TArray<Trace::FChannelEntry>& Channels = ChannelProvider->GetChannels();
+		for (uint64 ChannelIndex = 0; ChannelIndex < Channels.Num(); ++ChannelIndex)
+		{
+			if (Channels[ChannelIndex].bIsEnabled)
+			{
+				FrameZeroEnabledChannels.Add(Channels[ChannelIndex].Name);
+			}
+		}
+	}
+}
diff --git a/Engine/Plugins/Developer/TraceDataFiltering/Source/TraceDataFiltering/Private/BaseSessionFilterService.h b/Engine/Plugins/Developer/TraceDataFiltering/Source/TraceDataFiltering/Private/BaseSessionFilterService.h
new file mode 100644
--- /dev/null
+++ b/Engine/Plugins/Developer/TraceDataFiltering/Source/TraceDataFiltering/Private/BaseSessionFilterService.h
@@ -0,0 +1,54 @@
+// Copyright Epic Games, Inc. All Rights Reserved.
+
+#pragma once
+
+#include "ISessionTraceFilterService.h"
+#include "Misc/DateTime.h"
+
+namespace Trace
+{
+	class IAnalysisSession;
+	typedef uint64 FSessionHandle;
+}
+
+class IEventInfoProvider;
+
+/** Implementation of ISessionTraceFilterService specifically to use as a base for the Trace and Editor implementations, using Trace::IChannelProvider to provide information about Channels available on the running application 
+ and Trace::ISessionService to change the Channel state(s). */
+class FBaseSessionFilterService : public ISessionTraceFilterService
+{
+public:
+	FBaseSessionFilterService(Trace::FSessionHandle InHandle, TSharedPtr<const Trace::IAnalysisSession> InSession);
+	virtual ~FBaseSessionFilterService();
+
+	/** Begin ISessionTraceFilterService overrides */
+	virtual void GetRootObjects(TArray<FTraceObjectInfo>& OutObjects) const override;
+	virtual void GetChildObjects(uint32 InObjectHash, TArray<FTraceObjectInfo>& OutChildObjects) const override;
+	virtual const FDateTime& GetTimestamp() override;
+	virtual void SetObjectFilterState(const FString& InObjectName, const bool bFilterState) override;
+	virtual void UpdateFilterPresets(const TArray<TSharedPtr<IFilterPreset>>& InPresets) override;
+	/** End ISessionTraceFilterService overrides */
+
+protected:
+	/** Callback at end of engine frame, used to dispatch all enabled/disabled channels */
+	virtual void OnApplyChannelChanges() = 0;
+
+	/** Retrieves channels names from provider and marks them all as disabled */
+	void DisableAllChannels();
+
+	void RetrieveAndStoreStartupChannels();
+protected:
+	/** Session this instance represents the filtering service for */
+	TSharedPtr<const Trace::IAnalysisSession> Session;
+	/** Session handle for AnalyisSession*/
+	Trace::FSessionHandle Handle;
+
+	/** Names of channels that were either enabled or disabled during the duration of this frame */
+	TSet<FString> FrameEnabledChannels;
+	TSet<FString> FrameDisabledChannels;
+
+	TSet<FString> FrameZeroEnabledChannels;
+	
+	/** Timestamp at which contained data (including provider) was last updated */
+	FDateTime TimeStamp;	
+};
\ No newline at end of file
diff --git a/Engine/Plugins/Developer/TraceDataFiltering/Source/TraceDataFiltering/Private/STraceDataFilterWidget.cpp b/Engine/Plugins/Developer/TraceDataFiltering/Source/TraceDataFiltering/Private/STraceDataFilterWidget.cpp
--- a/Engine/Plugins/Developer/TraceDataFiltering/Source/TraceDataFiltering/Private/STraceDataFilterWidget.cpp
+++ b/Engine/Plugins/Developer/TraceDataFiltering/Source/TraceDataFiltering/Private/STraceDataFilterWidget.cpp
@@ -22,6 +22,7 @@
 #include "Insights/IUnrealInsightsModule.h"
 #include "Trace/StoreClient.h"
 #include "Trace/SessionTraceFilterService.h"
+#include "SessionEditorFilterService.h"
 #include "EventFilterStyle.h"
 
 #define LOCTEXT_NAMESPACE "STraceDataFilterWidget"
@@ -570,7 +571,11 @@ void STraceDataFilterWidget::RestoreItemSelection()
 
 void STraceDataFilterWidget::SetCurrentAnalysisSession(uint32 SessionHandle, TSharedRef<const Trace::IAnalysisSession> AnalysisSession)
 {
+#if WITH_EDITOR
+	SessionFilterService = MakeShareable(new FSessionEditorFilterService(SessionHandle, AnalysisSession));
+#else
 	SessionFilterService = MakeShareable(new FSessionTraceFilterService(SessionHandle, AnalysisSession));
+#endif
 
 	/** Refresh presets so config loaded state is directly applied */
 	OnPresetsChanged();
diff --git a/Engine/Plugins/Developer/TraceDataFiltering/Source/TraceDataFiltering/Private/SessionEditorFilterService.cpp b/Engine/Plugins/Developer/TraceDataFiltering/Source/TraceDataFiltering/Private/SessionEditorFilterService.cpp
new file mode 100644
--- /dev/null
+++ b/Engine/Plugins/Developer/TraceDataFiltering/Source/TraceDataFiltering/Private/SessionEditorFilterService.cpp
@@ -0,0 +1,33 @@
+// Copyright Epic Games, Inc. All Rights Reserved.
+
+#include "SessionEditorFilterService.h"
+#include "TraceServices/ITraceServicesModule.h"
+#include "Templates/SharedPointer.h"
+#include "Misc/CoreDelegates.h"
+#include "Trace/Trace.h"
+
+FSessionEditorFilterService::FSessionEditorFilterService(Trace::FSessionHandle InHandle, TSharedPtr<const Trace::IAnalysisSession> InSession) : FBaseSessionFilterService(InHandle, InSession)
+{	
+}
+
+void FSessionEditorFilterService::OnApplyChannelChanges()
+{
+	if (FrameEnabledChannels.Num())
+	{
+		for (const FString& ChannelName : FrameEnabledChannels)
+		{
+			Trace::ToggleChannel(*ChannelName, true);
+		}
+		FrameEnabledChannels.Empty();
+	}
+
+	if (FrameDisabledChannels.Num())
+	{
+		for (const FString& ChannelName : FrameDisabledChannels)
+		{
+			Trace::ToggleChannel(*ChannelName, false);
+		}
+
+		FrameDisabledChannels.Empty();
+	}
+}
diff --git a/Engine/Plugins/Developer/TraceDataFiltering/Source/TraceDataFiltering/Private/SessionEditorFilterService.h b/Engine/Plugins/Developer/TraceDataFiltering/Source/TraceDataFiltering/Private/SessionEditorFilterService.h
new file mode 100644
--- /dev/null
+++ b/Engine/Plugins/Developer/TraceDataFiltering/Source/TraceDataFiltering/Private/SessionEditorFilterService.h
@@ -0,0 +1,24 @@
+// Copyright Epic Games, Inc. All Rights Reserved.
+
+#pragma once
+
+#include "ISessionTraceFilterService.h"
+#include "BaseSessionFilterService.h"
+
+namespace Trace
+{
+	class IAnalysisSession;
+	typedef uint64 FSessionHandle;
+}
+
+/** Implementation of ISessionTraceFilterService specifically to use in the Editor */
+class FSessionEditorFilterService : public FBaseSessionFilterService
+{
+public:
+	FSessionEditorFilterService(Trace::FSessionHandle InHandle, TSharedPtr<const Trace::IAnalysisSession> InSession);
+	virtual ~FSessionEditorFilterService() {}
+
+	/** Begin FBaseSessionFilterService overrides */
+	virtual void OnApplyChannelChanges() override;
+	/** End OnApplyChannelChanges overrides */
+};
\ No newline at end of file
diff --git a/Engine/Plugins/Developer/TraceDataFiltering/Source/TraceDataFiltering/Private/Trace/SessionTraceFilterService.cpp b/Engine/Plugins/Developer/TraceDataFiltering/Source/TraceDataFiltering/Private/Trace/SessionTraceFilterService.cpp
--- a/Engine/Plugins/Developer/TraceDataFiltering/Source/TraceDataFiltering/Private/Trace/SessionTraceFilterService.cpp
+++ b/Engine/Plugins/Developer/TraceDataFiltering/Source/TraceDataFiltering/Private/Trace/SessionTraceFilterService.cpp
@@ -13,95 +13,12 @@
 #include <IPAddress.h>
 #include <SocketSubsystem.h>
 
-FSessionTraceFilterService::FSessionTraceFilterService(Trace::FSessionHandle InHandle, TSharedPtr<const Trace::IAnalysisSession> InSession) : Session(InSession), Handle(InHandle)
+FSessionTraceFilterService::FSessionTraceFilterService(Trace::FSessionHandle InHandle, TSharedPtr<const Trace::IAnalysisSession> InSession) : FBaseSessionFilterService(InHandle, InSession)
 {
-	FCoreDelegates::OnEndFrame.AddRaw(this, &FSessionTraceFilterService::OnEndFrame);
-	// Retrieve the channels currently enabled on the provider, this'll be the ones specified on the commandline (-trace=ChannelX)
-	RetrieveAndStoreStartupChannels();
-}
-
-FSessionTraceFilterService::~FSessionTraceFilterService()
-{
-	FCoreDelegates::OnEndFrame.RemoveAll(this);
-}
-
-void FSessionTraceFilterService::GetRootObjects(TArray<FTraceObjectInfo>& OutObjects) const
-{
-	const Trace::IChannelProvider* ChannelProvider = Session->ReadProvider<Trace::IChannelProvider>("ChannelProvider");
-	if (ChannelProvider)
-	{
-		const uint64 ChannelCount = ChannelProvider->GetChannelCount();		
-		const TArray<Trace::FChannelEntry>& Channels = ChannelProvider->GetChannels();
-		for (uint64 ChannelIndex = 0; ChannelIndex < Channels.Num(); ++ChannelIndex)
-		{
-			FTraceObjectInfo& EventInfo = OutObjects.AddDefaulted_GetRef();
-			EventInfo.Name = Channels[ChannelIndex].Name;
-			EventInfo.bEnabled = Channels[ChannelIndex].bIsEnabled;
-			EventInfo.Hash = GetTypeHash(EventInfo.Name);
-			EventInfo.OwnerHash = 0;
-		}
-	}
-}
 
-void FSessionTraceFilterService::GetChildObjects(uint32 InObjectHash, TArray<FTraceObjectInfo>& OutChildObjects) const
-{
-	/** TODO, parent/child relationship for Channels */
 }
 
-const FDateTime& FSessionTraceFilterService::GetTimestamp()
-{
-	const Trace::IChannelProvider* ChannelProvider = Session->ReadProvider<Trace::IChannelProvider>("ChannelProvider");
-	if (ChannelProvider)
-	{
-		TimeStamp = ChannelProvider->GetTimeStamp();
-	}
-
-	return TimeStamp;
-}
-
-void FSessionTraceFilterService::SetObjectFilterState(const FString& InObjectName, const bool bFilterState)
-{
-	if (bFilterState)
-	{
-		FrameDisabledChannels.Remove(InObjectName);
-		FrameEnabledChannels.Add(InObjectName);
-	}
-	else
-	{
-		FrameEnabledChannels.Remove(InObjectName);
-		FrameDisabledChannels.Add(InObjectName);
-	}
-}
-
-void FSessionTraceFilterService::UpdateFilterPresets(const TArray<TSharedPtr<IFilterPreset>>& InPresets)
-{
-	TSet<FString> UniqueNames;
-
-	for (const TSharedPtr<IFilterPreset>& Preset : InPresets)
-	{
-		TArray<FString> Names;
-		Preset->GetWhitelistedNames(Names);
-
-		// We are only interested in the unique channels names, as a result of combining multiple presets
-		Algo::Transform(Names, UniqueNames, [](FString InName) { return InName; });
-	}
-
-	DisableAllChannels();
-
-	for (const FString& EventName : UniqueNames)
-	{
-		SetObjectFilterState(EventName, true);
-	}
-
-	for (const FString& EventName : FrameZeroEnabledChannels)
-	{
-		SetObjectFilterState(EventName, true);
-	}
-	// Commandline channels are only applied once when changing presets
-	FrameZeroEnabledChannels.Empty();
-}
-
-void FSessionTraceFilterService::OnEndFrame()
+void FSessionTraceFilterService::OnApplyChannelChanges()
 {
 	auto GenerateConcatenatedChannels = [](TSet<FString>& InChannels, FString& OutConcatenation)
 	{
@@ -168,34 +85,3 @@ void FSessionTraceFilterService::OnEndFrame()
 
 	ControlClient.Disconnect();
 }
-
-void FSessionTraceFilterService::DisableAllChannels()
-{
-	const Trace::IChannelProvider* ChannelProvider = Session->ReadProvider<Trace::IChannelProvider>("ChannelProvider");
-	if (ChannelProvider)
-	{
-		const uint64 ChannelCount = ChannelProvider->GetChannelCount();
-		const TArray<Trace::FChannelEntry>& Channels = ChannelProvider->GetChannels();
-		for (uint64 ChannelIndex = 0; ChannelIndex < Channels.Num(); ++ChannelIndex)
-		{
-			SetObjectFilterState(Channels[ChannelIndex].Name, false);
-		}
-	}
-}
-
-void FSessionTraceFilterService::RetrieveAndStoreStartupChannels()
-{
-	const Trace::IChannelProvider* ChannelProvider = Session->ReadProvider<Trace::IChannelProvider>("ChannelProvider");	
-	if (ChannelProvider)
-	{
-		const uint64 ChannelCount = ChannelProvider->GetChannelCount();
-		const TArray<Trace::FChannelEntry>& Channels = ChannelProvider->GetChannels();
-		for (uint64 ChannelIndex = 0; ChannelIndex < Channels.Num(); ++ChannelIndex)
-		{
-			if (Channels[ChannelIndex].bIsEnabled)
-			{
-				FrameZeroEnabledChannels.Add(Channels[ChannelIndex].Name);
-			}
-		}
-	}
-}
diff --git a/Engine/Plugins/Developer/TraceDataFiltering/Source/TraceDataFiltering/Private/Trace/SessionTraceFilterService.h b/Engine/Plugins/Developer/TraceDataFiltering/Source/TraceDataFiltering/Private/Trace/SessionTraceFilterService.h
--- a/Engine/Plugins/Developer/TraceDataFiltering/Source/TraceDataFiltering/Private/Trace/SessionTraceFilterService.h
+++ b/Engine/Plugins/Developer/TraceDataFiltering/Source/TraceDataFiltering/Private/Trace/SessionTraceFilterService.h
@@ -3,7 +3,7 @@
 #pragma once
 
 #include "ISessionTraceFilterService.h"
-#include "Misc/DateTime.h"
+#include "BaseSessionFilterService.h"
 
 namespace Trace
 {
@@ -11,44 +11,15 @@ namespace Trace
 	typedef uint64 FSessionHandle;
 }
 
-class IEventInfoProvider;
-
 /** Implementation of ISessionTraceFilterService specifically to use with Trace, using Trace::IChannelProvider to provide information about Channels available on the running application 
  and Trace::ISessionService to change the Channel state(s). */
-class FSessionTraceFilterService : public ISessionTraceFilterService
+class FSessionTraceFilterService : public FBaseSessionFilterService
 {
 public:
 	FSessionTraceFilterService(Trace::FSessionHandle InHandle, TSharedPtr<const Trace::IAnalysisSession> InSession);
-	virtual ~FSessionTraceFilterService();
-
-	/** Begin ISessionTraceFilterService overrides */
-	virtual void GetRootObjects(TArray<FTraceObjectInfo>& OutObjects) const override;
-	virtual void GetChildObjects(uint32 InObjectHash, TArray<FTraceObjectInfo>& OutChildObjects) const override;
-	virtual const FDateTime& GetTimestamp() override;
-	virtual void SetObjectFilterState(const FString& InObjectName, const bool bFilterState) override;
-	virtual void UpdateFilterPresets(const TArray<TSharedPtr<IFilterPreset>>& InPresets) override;
-	/** End ISessionTraceFilterService overrides */
-
-protected:
-	/** Callback at end of engine frame, used to dispatch all enabled/disabled channels */
-	void OnEndFrame();
-
-	/** Retrieves channels names from provider and marks them all as disabled */
-	void DisableAllChannels();
-
-	void RetrieveAndStoreStartupChannels();
-protected:
-	/** Session this instance represents the filtering service for */
-	TSharedPtr<const Trace::IAnalysisSession> Session;
-	/** Session handle for AnalyisSession*/
-	Trace::FSessionHandle Handle;
-
-	/** Names of channels that were either enabled or disabled during the duration of this frame */
-	TSet<FString> FrameEnabledChannels;
-	TSet<FString> FrameDisabledChannels;
+	virtual ~FSessionTraceFilterService() {}
 
-	TSet<FString> FrameZeroEnabledChannels;
-	
-	/** Timestamp at which contained data (including provider) was last updated */
-	FDateTime TimeStamp;	
+	/** Begin FBaseSessionFilterService overrides */
+	virtual void OnApplyChannelChanges() override;
+	/** End OnApplyChannelChanges overrides */
 };
\ No newline at end of file
