[Insights] Fixed warning for Trace Data Filtering on non-live traces being illegible (color contrast).

diff --git a/Engine/Plugins/Developer/TraceDataFiltering/Source/TraceDataFiltering/Private/STraceDataFilterWidget.cpp b/Engine/Plugins/Developer/TraceDataFiltering/Source/TraceDataFiltering/Private/STraceDataFilterWidget.cpp
--- a/Engine/Plugins/Developer/TraceDataFiltering/Source/TraceDataFiltering/Private/STraceDataFilterWidget.cpp
+++ b/Engine/Plugins/Developer/TraceDataFiltering/Source/TraceDataFiltering/Private/STraceDataFilterWidget.cpp
@@ -129,7 +129,7 @@ void STraceDataFilterWidget::Construct(const FArguments& InArgs)
 			[
 				SNew(SBorder)
 				.BorderImage(FEventFilterStyle::GetBrush("FilterPresets.SessionWarningBorder"))
-				.BorderBackgroundColor(FColor(166,137,0))
+				.BorderBackgroundColor(FLinearColor(0.5f, 0.25f, 0.0f, 1.0f))
 				[
 					SNew(SHorizontalBox)
 					.Visibility_Lambda([this]() -> EVisibility 
@@ -146,13 +146,14 @@ void STraceDataFilterWidget::Construct(const FArguments& InArgs)
 						.Image(FEventFilterStyle::GetBrush("FilterPresets.WarningIcon"))
 					]
 
-					+SHorizontalBox::Slot()
+					+ SHorizontalBox::Slot()
 					.VAlign(VAlign_Center)
 					.AutoWidth()
 					.Padding(4.0f, 0.0f, 0.0f, 0.0f)
 					[
 						SNew(STextBlock)
 						.Text(LOCTEXT("NoValidFilterWarning", "Trace Data Filtering requires a connected Live Trace Session"))
+						.ColorAndOpacity(FLinearColor(1.0f, 1.0f, 1.0f, 1.0f))
 					]
 				]
 			]
@@ -167,7 +168,7 @@ void STraceDataFilterWidget::Construct(const FArguments& InArgs)
 				[
 					SNew(SScrollBorder, Treeview.ToSharedRef())
 					[
-						Treeview.ToSharedRef()					
+						Treeview.ToSharedRef()
 					]
 				]
 			]			
