Build Unreal Insights when launching from the Editor if the executable is not found

diff --git a/Engine/Plugins/Developer/TraceDataFiltering/Source/UnrealInsightsInterface/Private/UnrealInsightsLauncher.cpp b/Engine/Plugins/Developer/TraceDataFiltering/Source/UnrealInsightsInterface/Private/UnrealInsightsLauncher.cpp
--- a/Engine/Plugins/Developer/TraceDataFiltering/Source/UnrealInsightsInterface/Private/UnrealInsightsLauncher.cpp
+++ b/Engine/Plugins/Developer/TraceDataFiltering/Source/UnrealInsightsInterface/Private/UnrealInsightsLauncher.cpp
@@ -3,6 +3,7 @@
 #include "UnrealInsightsLauncher.h"
 
 #include "EditorStyleSet.h"
+#include "IUATHelperModule.h"
 #include "Logging/LogMacros.h"
 #include "Logging/MessageLog.h"
 #include "MessageLog/Public/MessageLogModule.h"
@@ -36,12 +37,12 @@ void FUnrealInsightsLauncher::RegisterMenus()
 			LOCTEXT("OpenUnrealInsights_Label", "Run Unreal Insights"),
 			LOCTEXT("OpenUnrealInsights_Desc", "Run the Unreal Insights standalone application."),
 			FSlateIcon(FEditorStyle::GetStyleSetName(), "UnrealInsights.MenuIcon"),
-			FUIAction(FExecuteAction::CreateRaw(this, &FUnrealInsightsLauncher::OpenUnrealInsights), FCanExecuteAction())
+			FUIAction(FExecuteAction::CreateRaw(this, &FUnrealInsightsLauncher::RunUnrealInsights_Execute), FCanExecuteAction())
 		);
 	}
 }
 
-void FUnrealInsightsLauncher::OpenUnrealInsights()
+void FUnrealInsightsLauncher::RunUnrealInsights_Execute()
 {
 	FString Path = FPaths::GetPath(FPlatformProcess::ExecutablePath());
 	Path = FPaths::Combine(Path, TEXT("UnrealInsights"));
@@ -50,6 +51,45 @@ void FUnrealInsightsLauncher::OpenUnrealInsights()
 #endif
 	FPaths::MakeStandardFilename(Path);
 
+	FMessageLogModule& MessageLogModule = FModuleManager::LoadModuleChecked<FMessageLogModule>("MessageLog");
+	if (!MessageLogModule.IsRegisteredLogListing(LogListingName))
+	{
+		MessageLogModule.RegisterLogListing(LogListingName, LOCTEXT("UnrealInsights", "Unreal Insights"));
+	}
+
+	if (!FPaths::FileExists(Path))
+	{
+		UE_LOG(UnrealInsightsInterface, Log, TEXT("Could not find the Unreal Insights executable: %s. Attempting to build UnrealInsights."), *Path);
+
+		FString Arguments;
+#if PLATFORM_WINDOWS
+		FText PlatformName = LOCTEXT("PlatformName_Windows", "Windows");
+		Arguments = TEXT("BuildTarget -Target=UnrealInsights -Platform=Win64");
+#elif PLATFORM_MAC
+		FText PlatformName = LOCTEXT("PlatformName_Mac", "Mac");
+		Arguments = TEXT("BuildTarget -Target=UnrealInsights -Platform=Mac");
+#elif PLATFORM_LINUX
+		FText PlatformName = LOCTEXT("PlatformName_Linux", "Linux");
+		Arguments = TEXT("BuildTarget -Target=UnrealInsights -Platform=Linux");
+#endif
+
+		IUATHelperModule::Get().CreateUatTask(Arguments, PlatformName, LOCTEXT("BuildingUnrealInsights", "Building Unreal Insights"),
+			LOCTEXT("BuildUnrealInsightsTask", "Build Unreal Insights Task"), FEditorStyle::GetBrush(TEXT("MainFrame.CookContent")), [this, Path](FString Result, double Time)
+			{
+				if (Result.Equals(TEXT("Completed")))
+				{
+					this->StartUnrealInsights(Path);
+				}
+			});
+	}
+	else
+	{
+		StartUnrealInsights(Path);
+	}
+}
+
+void FUnrealInsightsLauncher::StartUnrealInsights(const FString& Path)
+{
 	FString CmdLine;
 
 	constexpr bool bLaunchDetached = true;
@@ -70,21 +110,7 @@ void FUnrealInsightsLauncher::OpenUnrealInsights()
 	}
 	else
 	{
-		FMessageLogModule& MessageLogModule = FModuleManager::LoadModuleChecked<FMessageLogModule>("MessageLog");
-		if (!MessageLogModule.IsRegisteredLogListing(LogListingName))
-		{
-			MessageLogModule.RegisterLogListing(LogListingName, LOCTEXT("UnrealInsights", "Unreal Insights"));
-		}
-
-		FText MessageBoxTextFmt;
-		if (!FPaths::FileExists(Path))
-		{
-			MessageBoxTextFmt = LOCTEXT("ExecutableNotFound_TextFmt", "Unreal Insights executable could not be found at: {0}. Is Unreal Insights built?");
-		}
-		else
-		{
-			MessageBoxTextFmt = LOCTEXT("ExecutableNotFound_TextFmt", "Could not start Unreal Insights executable at path: {0}");
-		}
+		const FText	MessageBoxTextFmt = LOCTEXT("ExecutableNotFound_TextFmt", "Could not start Unreal Insights executable at path: {0}");
 		const FText MessageBoxText = FText::Format(MessageBoxTextFmt, FText::FromString(Path));
 
 		FMessageLog ReportMessageLog(LogListingName);
@@ -92,4 +118,4 @@ void FUnrealInsightsLauncher::OpenUnrealInsights()
 		ReportMessageLog.AddMessage(Message);
 		ReportMessageLog.Notify();
 	}
-}
\ No newline at end of file
+}
diff --git a/Engine/Plugins/Developer/TraceDataFiltering/Source/UnrealInsightsInterface/Private/UnrealInsightsLauncher.h b/Engine/Plugins/Developer/TraceDataFiltering/Source/UnrealInsightsInterface/Private/UnrealInsightsLauncher.h
--- a/Engine/Plugins/Developer/TraceDataFiltering/Source/UnrealInsightsInterface/Private/UnrealInsightsLauncher.h
+++ b/Engine/Plugins/Developer/TraceDataFiltering/Source/UnrealInsightsInterface/Private/UnrealInsightsLauncher.h
@@ -11,9 +11,10 @@ public:
 	~FUnrealInsightsLauncher();
 
 	void RegisterMenus();
+	void StartUnrealInsights(const FString& Path);
 
 private:
-	void OpenUnrealInsights();
+	void RunUnrealInsights_Execute();
 
 private:
 	/** The name of the Unreal Insights log listing. */
diff --git a/Engine/Plugins/Developer/TraceDataFiltering/Source/UnrealInsightsInterface/UnrealInsightsInterface.Build.cs b/Engine/Plugins/Developer/TraceDataFiltering/Source/UnrealInsightsInterface/UnrealInsightsInterface.Build.cs
--- a/Engine/Plugins/Developer/TraceDataFiltering/Source/UnrealInsightsInterface/UnrealInsightsInterface.Build.cs
+++ b/Engine/Plugins/Developer/TraceDataFiltering/Source/UnrealInsightsInterface/UnrealInsightsInterface.Build.cs
@@ -20,6 +20,7 @@ namespace UnrealBuildTool.Rules
 					"EditorFramework",
 					"ApplicationCore",
 					"ToolMenus",
+					"UATHelper",
 				}
 			);
 		}
