Fix UnrealInsightsInterface module name for monolithic build

diff --git a/Engine/Plugins/Developer/TraceDataFiltering/Source/UnrealInsightsInterface/Private/UnrealInsightsInterfaceModule.cpp b/Engine/Plugins/Developer/TraceDataFiltering/Source/UnrealInsightsInterface/Private/UnrealInsightsInterfaceModule.cpp
--- a/Engine/Plugins/Developer/TraceDataFiltering/Source/UnrealInsightsInterface/Private/UnrealInsightsInterfaceModule.cpp
+++ b/Engine/Plugins/Developer/TraceDataFiltering/Source/UnrealInsightsInterface/Private/UnrealInsightsInterfaceModule.cpp
@@ -11,7 +11,7 @@
 
 DEFINE_LOG_CATEGORY(UnrealInsightsInterface);
 
-IMPLEMENT_MODULE(FUnrealInsightsInterfaceModule, UnrealInsights);
+IMPLEMENT_MODULE(FUnrealInsightsInterfaceModule, UnrealInsightsInterface);
 
 void FUnrealInsightsInterfaceModule::StartupModule()
 {
