Replace uses of inappropriate words in TraceDataFiltering module

diff --git a/Engine/Plugins/Developer/TraceDataFiltering/Config/BaseTraceDataFilters.ini b/Engine/Plugins/Developer/TraceDataFiltering/Config/BaseTraceDataFilters.ini
--- a/Engine/Plugins/Developer/TraceDataFiltering/Config/BaseTraceDataFilters.ini
+++ b/Engine/Plugins/Developer/TraceDataFiltering/Config/BaseTraceDataFilters.ini
@@ -1,2 +1,2 @@
 [/Script/TraceDataFiltering.EngineFilterPresetContainer]
-+EnginePresets=(Name="Animation",WhitelistedNames=("Animation", "Object", "Frame"))
++EnginePresets=(Name="Animation",AllowlistedNames=("Animation", "Object", "Frame"))
diff --git a/Engine/Plugins/Developer/TraceDataFiltering/Source/TraceDataFiltering/Private/BaseSessionFilterService.cpp b/Engine/Plugins/Developer/TraceDataFiltering/Source/TraceDataFiltering/Private/BaseSessionFilterService.cpp
--- a/Engine/Plugins/Developer/TraceDataFiltering/Source/TraceDataFiltering/Private/BaseSessionFilterService.cpp
+++ b/Engine/Plugins/Developer/TraceDataFiltering/Source/TraceDataFiltering/Private/BaseSessionFilterService.cpp
@@ -75,7 +75,7 @@ void FBaseSessionFilterService::SetObjectFilterState(const FString& InObjectName
 void FBaseSessionFilterService::UpdateFilterPreset(const TSharedPtr<IFilterPreset> Preset, bool IsEnabled)
 {
 	TArray<FString> Names;
-	Preset->GetWhitelistedNames(Names);
+	Preset->GetAllowlistedNames(Names);
 	if (IsEnabled)
 	{
 		FrameEnabledChannels.Append(Names);
diff --git a/Engine/Plugins/Developer/TraceDataFiltering/Source/TraceDataFiltering/Private/FilterPresets.cpp b/Engine/Plugins/Developer/TraceDataFiltering/Source/TraceDataFiltering/Private/FilterPresets.cpp
--- a/Engine/Plugins/Developer/TraceDataFiltering/Source/TraceDataFiltering/Private/FilterPresets.cpp
+++ b/Engine/Plugins/Developer/TraceDataFiltering/Source/TraceDataFiltering/Private/FilterPresets.cpp
@@ -112,7 +112,7 @@ void FFilterPresetHelpers::CreateNewPreset(const TArray<TSharedPtr<ITraceObject>
 
 	TArray<FString> Names;
 	ExtractEnabledObjectNames(InObjects, Names);
-	NewUserFilter.WhitelistedNames = Names;
+	NewUserFilter.AllowlistedNames = Names;
 
 	ULocalFilterPresetContainer::Save();
 }
@@ -200,7 +200,7 @@ void FUserFilterPreset::Save(const TArray<TSharedPtr<ITraceObject>>& InObjects)
 	TArray<FString> Names;
 	FFilterPresetHelpers::ExtractEnabledObjectNames(InObjects, Names);
 
-	FilterData.WhitelistedNames = Names;
+	FilterData.AllowlistedNames = Names;
 	
 	USharedFilterPresetContainer::Save();
 	ULocalFilterPresetContainer::Save();
@@ -227,9 +227,9 @@ FText FFilterPreset::GetDescription() const
 	return FText::FormatOrdered(LOCTEXT("FilterPresetDescriptionFormat", "Name: {0}\nType: {1}"), FText::FromString(Name), CanDelete() ? (IsLocal() ? LOCTEXT("LocalPreset", "Local") : LOCTEXT("SharedPreset", "Shared")) : LOCTEXT("EnginePreset", "Engine"));
 }
 
-void FFilterPreset::GetWhitelistedNames(TArray<FString>& OutNames) const
+void FFilterPreset::GetAllowlistedNames(TArray<FString>& OutNames) const
 {
-	OutNames.Append(FilterData.WhitelistedNames);
+	OutNames.Append(FilterData.AllowlistedNames);
 }
 
 bool FFilterPreset::CanDelete() const
diff --git a/Engine/Plugins/Developer/TraceDataFiltering/Source/TraceDataFiltering/Private/FilterPresets.h b/Engine/Plugins/Developer/TraceDataFiltering/Source/TraceDataFiltering/Private/FilterPresets.h
--- a/Engine/Plugins/Developer/TraceDataFiltering/Source/TraceDataFiltering/Private/FilterPresets.h
+++ b/Engine/Plugins/Developer/TraceDataFiltering/Source/TraceDataFiltering/Private/FilterPresets.h
@@ -20,11 +20,11 @@ struct FFilterData
 	FString Name;
 
 	UPROPERTY()
-	TArray<FString> WhitelistedNames;
+	TArray<FString> AllowlistedNames;
 	
 	bool operator==(const FFilterData& Other) const
 	{
-		return Name == Other.Name && WhitelistedNames == Other.WhitelistedNames;
+		return Name == Other.Name && AllowlistedNames == Other.AllowlistedNames;
 	}
 };
 
@@ -48,7 +48,7 @@ public:
 	virtual FString GetName() const override;
 	virtual FText GetDisplayText() const;
 	virtual FText GetDescription() const;
-	virtual void GetWhitelistedNames(TArray<FString>& OutNames) const override;
+	virtual void GetAllowlistedNames(TArray<FString>& OutNames) const override;
 	virtual bool CanDelete() const override;
 	virtual void Rename(const FString& InNewName) override;
 	virtual bool Delete() override;
diff --git a/Engine/Plugins/Developer/TraceDataFiltering/Source/TraceDataFiltering/Private/IFilterPreset.h b/Engine/Plugins/Developer/TraceDataFiltering/Source/TraceDataFiltering/Private/IFilterPreset.h
--- a/Engine/Plugins/Developer/TraceDataFiltering/Source/TraceDataFiltering/Private/IFilterPreset.h
+++ b/Engine/Plugins/Developer/TraceDataFiltering/Source/TraceDataFiltering/Private/IFilterPreset.h
@@ -17,7 +17,7 @@ public:
 	virtual FText GetDescription() const = 0;
 
 	/** Returns all names which pass this preset its filter */
-	virtual void GetWhitelistedNames(TArray<FString>& OutNames) const = 0;
+	virtual void GetAllowlistedNames(TArray<FString>& OutNames) const = 0;
 	/** Whether or not this preset can be deleted */
 	virtual bool CanDelete() const = 0;
 	/** Whether or not this preset only exists locally */
diff --git a/Engine/Plugins/Developer/TraceDataFiltering/Source/TraceDataFiltering/Private/SFilterPreset.cpp b/Engine/Plugins/Developer/TraceDataFiltering/Source/TraceDataFiltering/Private/SFilterPreset.cpp
--- a/Engine/Plugins/Developer/TraceDataFiltering/Source/TraceDataFiltering/Private/SFilterPreset.cpp
+++ b/Engine/Plugins/Developer/TraceDataFiltering/Source/TraceDataFiltering/Private/SFilterPreset.cpp
@@ -165,7 +165,7 @@ FReply SFilterPreset::FilterMiddleButtonClicked()
 
 FText SFilterPreset::GetToolTipText() const
 {
-	return bHighlighted ? FilterPreset->GetDescription() : FText::Format(LOCTEXT("FilterPresetTooltipText", "{0}\nHold CTRL to highlight whitelisted items."), FilterPreset->GetDescription());
+	return bHighlighted ? FilterPreset->GetDescription() : FText::Format(LOCTEXT("FilterPresetTooltipText", "{0}\nHold CTRL to highlight allowlisted items."), FilterPreset->GetDescription());
 }
 
 TSharedRef<SWidget> SFilterPreset::GetRightClickMenuContent()
diff --git a/Engine/Plugins/Developer/TraceDataFiltering/Source/TraceDataFiltering/Private/STraceDataFilterWidget.cpp b/Engine/Plugins/Developer/TraceDataFiltering/Source/TraceDataFiltering/Private/STraceDataFilterWidget.cpp
--- a/Engine/Plugins/Developer/TraceDataFiltering/Source/TraceDataFiltering/Private/STraceDataFilterWidget.cpp
+++ b/Engine/Plugins/Developer/TraceDataFiltering/Source/TraceDataFiltering/Private/STraceDataFilterWidget.cpp
@@ -216,7 +216,7 @@ void STraceDataFilterWidget::OnHighlightPreset(const TSharedPtr<IFilterPreset>&
 {
 	if (Treeview.IsValid())
 	{
-		/** Update treeview so that any whitelisted entry (as part of Preset) is highlighted and expanded */
+		/** Update treeview so that any allowlisted entry (as part of Preset) is highlighted and expanded */
 		Treeview->ClearHighlightedItems();
 		if (Preset.IsValid())
 		{
@@ -228,7 +228,7 @@ void STraceDataFilterWidget::OnHighlightPreset(const TSharedPtr<IFilterPreset>&
 			}
 
 			TArray<FString> Names;
-			Preset->GetWhitelistedNames(Names);
+			Preset->GetAllowlistedNames(Names);
 
 			EnumerateAllItems([this, Names](TSharedPtr<ITraceObject> Object) -> void
 			{
