Update TraceDataFiltering warning message to UE5 style

diff --git a/Engine/Plugins/Developer/TraceDataFiltering/Source/TraceDataFiltering/Private/EventFilterStyle.cpp b/Engine/Plugins/Developer/TraceDataFiltering/Source/TraceDataFiltering/Private/EventFilterStyle.cpp
--- a/Engine/Plugins/Developer/TraceDataFiltering/Source/TraceDataFiltering/Private/EventFilterStyle.cpp
+++ b/Engine/Plugins/Developer/TraceDataFiltering/Source/TraceDataFiltering/Private/EventFilterStyle.cpp
@@ -1,6 +1,7 @@
 // Copyright Epic Games, Inc. All Rights Reserved.
 
 #include "EventFilterStyle.h"
+#include "Styling/StyleColors.h"
 #include "Styling/SlateStyleRegistry.h"
 #include "Styling/SlateTypes.h"
 #include "Styling/CoreStyle.h"
@@ -9,6 +10,7 @@
 
 // This is to fix the issue that SlateStyleMacros like IMAGE_BRUSH look for RootToContentDir but StyleSet->RootToContentDir is how this style is set up
 #define RootToContentDir StyleSet->RootToContentDir
+#define RootToCoreContentDir StyleSet->RootToCoreContentDir
 
 TSharedPtr< FSlateStyleSet > FEventFilterStyle::StyleSet = nullptr;
 
@@ -122,8 +124,8 @@ void FEventFilterStyle::Initialize()
 	StyleSet->Set("FilterPresets.FilterNameFont", DEFAULT_FONT("Regular", 10));
 	StyleSet->Set("FilterPresets.FilterButtonBorder", new BOX_BRUSH("Common/RoundedSelection_16x", FMargin(4.0f / 16.0f)));
 
-	StyleSet->Set("FilterPresets.SessionWarningBorder", new BOX_BRUSH("Common/GroupBorderLight", FMargin(4.0f / 16.0f)));
-	StyleSet->Set("FilterPresets.WarningIcon", new IMAGE_BRUSH("Settings/Settings_Warning", Icon40x40));
+	StyleSet->Set("FilterPresets.SessionWarningBorder", new BOX_BRUSH("Common/DarkGroupBorder", FMargin(4.0f / 16.0f)));
+	StyleSet->Set("FilterPresets.WarningIcon", new CORE_IMAGE_BRUSH_SVG("Starship/Common/alert-circle", Icon40x40, FStyleColors::Warning));
 
 	StyleSet->Set("FontAwesome.9", ICON_FONT(9));
 	
diff --git a/Engine/Plugins/Developer/TraceDataFiltering/Source/TraceDataFiltering/Private/STraceDataFilterWidget.cpp b/Engine/Plugins/Developer/TraceDataFiltering/Source/TraceDataFiltering/Private/STraceDataFilterWidget.cpp
--- a/Engine/Plugins/Developer/TraceDataFiltering/Source/TraceDataFiltering/Private/STraceDataFilterWidget.cpp
+++ b/Engine/Plugins/Developer/TraceDataFiltering/Source/TraceDataFiltering/Private/STraceDataFilterWidget.cpp
@@ -129,7 +129,6 @@ void STraceDataFilterWidget::Construct(const FArguments& InArgs)
 			[
 				SNew(SBorder)
 				.BorderImage(FEventFilterStyle::GetBrush("FilterPresets.SessionWarningBorder"))
-				.BorderBackgroundColor(FLinearColor(0.5f, 0.25f, 0.0f, 1.0f))
 				[
 					SNew(SHorizontalBox)
 					.Visibility_Lambda([this]() -> EVisibility 
