Fix launching Unreal Insights from Editor on Mac.

diff --git a/Engine/Plugins/Developer/TraceDataFiltering/Source/UnrealInsightsInterface/Private/UnrealInsightsLauncher.cpp b/Engine/Plugins/Developer/TraceDataFiltering/Source/UnrealInsightsInterface/Private/UnrealInsightsLauncher.cpp
--- a/Engine/Plugins/Developer/TraceDataFiltering/Source/UnrealInsightsInterface/Private/UnrealInsightsLauncher.cpp
+++ b/Engine/Plugins/Developer/TraceDataFiltering/Source/UnrealInsightsInterface/Private/UnrealInsightsLauncher.cpp
@@ -44,12 +44,7 @@ void FUnrealInsightsLauncher::RegisterMenus()
 
 void FUnrealInsightsLauncher::RunUnrealInsights_Execute()
 {
-	FString Path = FPaths::GetPath(FPlatformProcess::ExecutablePath());
-	Path = FPaths::Combine(Path, TEXT("UnrealInsights"));
-#if PLATFORM_WINDOWS
-	Path = FPaths::SetExtension(Path, TEXT(".exe"));
-#endif
-	FPaths::MakeStandardFilename(Path);
+	FString Path = FPlatformProcess::GenerateApplicationPath(TEXT("UnrealInsights"), EBuildConfiguration::Development);
 
 	FMessageLogModule& MessageLogModule = FModuleManager::LoadModuleChecked<FMessageLogModule>("MessageLog");
 	if (!MessageLogModule.IsRegisteredLogListing(LogListingName))
@@ -59,6 +54,7 @@ void FUnrealInsightsLauncher::RunUnrealInsights_Execute()
 
 	if (!FPaths::FileExists(Path))
 	{
+#if !PLATFORM_MAC
 		UE_LOG(UnrealInsightsInterface, Log, TEXT("Could not find the Unreal Insights executable: %s. Attempting to build UnrealInsights."), *Path);
 
 		FString Arguments;
@@ -81,6 +77,15 @@ void FUnrealInsightsLauncher::RunUnrealInsights_Execute()
 					this->StartUnrealInsights(Path);
 				}
 			});
+#else
+		const FText	MessageBoxTextFmt = LOCTEXT("ExecutableNotFoundManualBuild_TextFmt", "Could not find Unreal Insights executable. Have you built Unreal Insights?");
+		const FText MessageBoxText = FText::Format(MessageBoxTextFmt, FText::FromString(Path));
+
+		FMessageLog ReportMessageLog(LogListingName);
+		TSharedRef<FTokenizedMessage> Message = FTokenizedMessage::Create(EMessageSeverity::Error, MessageBoxText);
+		ReportMessageLog.AddMessage(Message);
+		ReportMessageLog.Notify();
+#endif
 	}
 	else
 	{
