Fixed Unreal Insights typo in the Tools Menu

diff --git a/Engine/Plugins/Developer/TraceDataFiltering/Source/UnrealInsightsInterface/Private/UnrealInsightsLauncher.cpp b/Engine/Plugins/Developer/TraceDataFiltering/Source/UnrealInsightsInterface/Private/UnrealInsightsLauncher.cpp
--- a/Engine/Plugins/Developer/TraceDataFiltering/Source/UnrealInsightsInterface/Private/UnrealInsightsLauncher.cpp
+++ b/Engine/Plugins/Developer/TraceDataFiltering/Source/UnrealInsightsInterface/Private/UnrealInsightsLauncher.cpp
@@ -32,7 +32,7 @@ void FUnrealInsightsLauncher::RegisterMenus()
 	UToolMenu* ProfileMenu = UToolMenus::Get()->ExtendMenu("MainFrame.MainMenu.Tools");
 	if (ProfileMenu)
 	{
-		FToolMenuSection& Section = ProfileMenu->AddSection("Unreal Insights", FText::FromString(TEXT("Unreal Unsights")));
+		FToolMenuSection& Section = ProfileMenu->AddSection("Unreal Insights", FText::FromString(TEXT("Unreal Insights")));
 		Section.AddMenuEntry("OpenUnrealInsights",
 			LOCTEXT("OpenUnrealInsights_Label", "Run Unreal Insights"),
 			LOCTEXT("OpenUnrealInsights_Desc", "Run the Unreal Insights standalone application."),
