[Insights] Added more "Insights/TraceDataFiltering" memory scopes in the TraceDataFiltering module (for untagged memory allocations while running in Editor). Also added "Insights" mem scope for allocs made from UnrealInsightsInterface module.

diff --git a/Engine/Plugins/Developer/TraceDataFiltering/Source/TraceDataFiltering/Private/EventFilterStyle.cpp b/Engine/Plugins/Developer/TraceDataFiltering/Source/TraceDataFiltering/Private/EventFilterStyle.cpp
--- a/Engine/Plugins/Developer/TraceDataFiltering/Source/TraceDataFiltering/Private/EventFilterStyle.cpp
+++ b/Engine/Plugins/Developer/TraceDataFiltering/Source/TraceDataFiltering/Private/EventFilterStyle.cpp
@@ -1,6 +1,8 @@
 // Copyright Epic Games, Inc. All Rights Reserved.
 
 #include "EventFilterStyle.h"
+
+#include "HAL/LowLevelMemTracker.h"
 #include "Styling/StyleColors.h"
 #include "Styling/SlateStyleRegistry.h"
 #include "Styling/SlateTypes.h"
@@ -40,6 +42,8 @@ static const FVector2D Icon267x140(170.0f, 50.0f);
 
 void FEventFilterStyle::Initialize()
 {
+	LLM_SCOPE_BYNAME(TEXT("Insights/TraceDataFiltering"));
+
 	// Only register once
 	if( StyleSet.IsValid() )
 	{
diff --git a/Engine/Plugins/Developer/TraceDataFiltering/Source/TraceDataFiltering/Private/STraceDataFilterWidget.cpp b/Engine/Plugins/Developer/TraceDataFiltering/Source/TraceDataFiltering/Private/STraceDataFilterWidget.cpp
--- a/Engine/Plugins/Developer/TraceDataFiltering/Source/TraceDataFiltering/Private/STraceDataFilterWidget.cpp
+++ b/Engine/Plugins/Developer/TraceDataFiltering/Source/TraceDataFiltering/Private/STraceDataFilterWidget.cpp
@@ -1,6 +1,7 @@
 // Copyright Epic Games, Inc. All Rights Reserved.
 
 #include "STraceDataFilterWidget.h"
+#include "HAL/LowLevelMemTracker.h"
 #include "Styling/AppStyle.h"
 #include "Widgets/Layout/SScrollBorder.h"
 #include "Widgets/Text/STextBlock.h"
@@ -318,6 +319,7 @@ void STraceDataFilterWidget::ConstructFilterHandler()
 
 TSharedRef<ITableRow> STraceDataFilterWidget::OnGenerateRow(TSharedPtr<ITraceObject> InItem, const TSharedRef<STableViewBase>& OwnerTable)
 {
+	LLM_SCOPE_BYNAME(TEXT("Insights/TraceDataFiltering"));
 	return SNew(STraceObjectRowWidget, OwnerTable, InItem)
 		.HighlightText(SearchBoxWidgetFilter.Get(), &TTextFilter<TSharedPtr<ITraceObject>>::GetRawFilterText);
 }
@@ -716,6 +718,8 @@ void STraceDataFilterWidget::RefreshTreeviewData()
 
 void STraceDataFilterWidget::Tick(const FGeometry& AllottedGeometry, const double InCurrentTime, const float InDeltaTime)
 {
+	LLM_SCOPE_BYNAME(TEXT("Insights/TraceDataFiltering"));
+
 	IUnrealInsightsModule& InsightsModule = FModuleManager::LoadModuleChecked<IUnrealInsightsModule>("TraceInsights");
 	TSharedPtr<const TraceServices::IAnalysisSession> AnalysisSession = InsightsModule.GetAnalysisSession();
 
diff --git a/Engine/Plugins/Developer/TraceDataFiltering/Source/TraceDataFiltering/Private/TraceDataFilteringModule.cpp b/Engine/Plugins/Developer/TraceDataFiltering/Source/TraceDataFiltering/Private/TraceDataFilteringModule.cpp
--- a/Engine/Plugins/Developer/TraceDataFiltering/Source/TraceDataFiltering/Private/TraceDataFilteringModule.cpp
+++ b/Engine/Plugins/Developer/TraceDataFiltering/Source/TraceDataFiltering/Private/TraceDataFilteringModule.cpp
@@ -2,6 +2,7 @@
 
 #include "TraceDataFilteringModule.h"
 
+#include "HAL/LowLevelMemTracker.h"
 #include "Modules/ModuleInterface.h"
 #include "Modules/ModuleManager.h"
 #include "Features/IModularFeatures.h"
@@ -30,6 +31,8 @@ FString FTraceFilteringModule::TraceFiltersIni;
 
 void FTraceFilteringModule::StartupModule()
 {
+	LLM_SCOPE_BYNAME(TEXT("Insights/TraceDataFiltering"));
+
 	FEventFilterStyle::Initialize();
 
 	FConfigCacheIni::LoadGlobalIniFile(TraceFiltersIni, TEXT("TraceDataFilters"));
@@ -39,6 +42,8 @@ void FTraceFilteringModule::StartupModule()
 	FTabSpawnerEntry& FilterTabSpawnerEntry = FGlobalTabmanager::Get()->RegisterNomadTabSpawner(FTraceFilteringModule::InsightsFilterTabName,
 		FOnSpawnTab::CreateLambda([](const FSpawnTabArgs& Args)
 		{
+			LLM_SCOPE_BYNAME(TEXT("Insights/TraceDataFiltering"));
+
 			const TSharedRef<SDockTab> DockTab = SNew(SDockTab)
 				.TabRole(ETabRole::NomadTab);
 
@@ -62,7 +67,10 @@ void FTraceFilteringModule::StartupModule()
 
 void FTraceFilteringModule::ShutdownModule()
 {
+	LLM_SCOPE_BYNAME(TEXT("Insights/TraceDataFiltering"));
+
 	FEventFilterStyle::Shutdown();
+
 #if WITH_EDITOR
 	FGlobalTabmanager::Get()->UnregisterNomadTabSpawner(FTraceFilteringModule::InsightsFilterTabName);
 #else
@@ -84,6 +92,8 @@ void FTraceFilteringModule::RegisterTimingProfilerLayoutExtensions(FInsightsMajo
 	MinorTabConfig.WorkspaceGroup = Category;
 	MinorTabConfig.OnSpawnTab = FOnSpawnTab::CreateLambda([](const FSpawnTabArgs& Args)
 	{
+		LLM_SCOPE_BYNAME(TEXT("Insights/TraceDataFiltering"));
+
 		const TSharedRef<SDockTab> DockTab = SNew(SDockTab)
 			.TabRole(ETabRole::NomadTab);
 
diff --git a/Engine/Plugins/Developer/TraceDataFiltering/Source/UnrealInsightsInterface/Private/UnrealInsightsInterfaceModule.cpp b/Engine/Plugins/Developer/TraceDataFiltering/Source/UnrealInsightsInterface/Private/UnrealInsightsInterfaceModule.cpp
--- a/Engine/Plugins/Developer/TraceDataFiltering/Source/UnrealInsightsInterface/Private/UnrealInsightsInterfaceModule.cpp
+++ b/Engine/Plugins/Developer/TraceDataFiltering/Source/UnrealInsightsInterface/Private/UnrealInsightsInterfaceModule.cpp
@@ -2,8 +2,9 @@
 
 #include "UnrealInsightsInterfaceModule.h"
 
-#include "Styling/AppStyle.h"
+#include "HAL/LowLevelMemTracker.h"
 #include "Logging/MessageLog.h"
+#include "Styling/AppStyle.h"
 #include "ToolMenus.h"
 #include "Widgets/Docking/SDockTab.h"
 
@@ -15,11 +16,13 @@ IMPLEMENT_MODULE(FUnrealInsightsInterfaceModule, UnrealInsightsInterface);
 
 void FUnrealInsightsInterfaceModule::StartupModule()
 {
+	LLM_SCOPE_BYNAME(TEXT("Insights"));
 	Launcher = MakeShared<FUnrealInsightsLauncher>();
 	RegisterStartupCallbackHandle = UToolMenus::RegisterStartupCallback(FSimpleMulticastDelegate::FDelegate::CreateSP(Launcher.Get(), &FUnrealInsightsLauncher::RegisterMenus));
 }
 
 void FUnrealInsightsInterfaceModule::ShutdownModule()
 {
+	LLM_SCOPE_BYNAME(TEXT("Insights"));
 	UToolMenus::UnRegisterStartupCallback(RegisterStartupCallbackHandle);
 }
