Use the absolute path when launching Unreal Insights from the editor menu.

diff --git a/Engine/Plugins/Developer/TraceDataFiltering/Source/UnrealInsightsInterface/Private/UnrealInsightsLauncher.cpp b/Engine/Plugins/Developer/TraceDataFiltering/Source/UnrealInsightsInterface/Private/UnrealInsightsLauncher.cpp
--- a/Engine/Plugins/Developer/TraceDataFiltering/Source/UnrealInsightsInterface/Private/UnrealInsightsLauncher.cpp
+++ b/Engine/Plugins/Developer/TraceDataFiltering/Source/UnrealInsightsInterface/Private/UnrealInsightsLauncher.cpp
@@ -45,7 +45,8 @@ void FUnrealInsightsLauncher::RegisterMenus()
 void FUnrealInsightsLauncher::RunUnrealInsights_Execute()
 {
 	FString Path = FPlatformProcess::GenerateApplicationPath(TEXT("UnrealInsights"), EBuildConfiguration::Development);
-
+	Path = FPaths::ConvertRelativePathToFull(Path);
+	
 	FMessageLogModule& MessageLogModule = FModuleManager::LoadModuleChecked<FMessageLogModule>("MessageLog");
 	if (!MessageLogModule.IsRegisteredLogListing(LogListingName))
 	{
