[Insights] Fixed conversion warnings in the TraceDataFiltering plugin.

diff --git a/Engine/Plugins/Developer/TraceDataFiltering/Source/TraceDataFiltering/Private/BaseSessionFilterService.cpp b/Engine/Plugins/Developer/TraceDataFiltering/Source/TraceDataFiltering/Private/BaseSessionFilterService.cpp
--- a/Engine/Plugins/Developer/TraceDataFiltering/Source/TraceDataFiltering/Private/BaseSessionFilterService.cpp
+++ b/Engine/Plugins/Developer/TraceDataFiltering/Source/TraceDataFiltering/Private/BaseSessionFilterService.cpp
@@ -20,9 +20,8 @@ void FBaseSessionFilterService::GetRootObjects(TArray<FTraceObjectInfo>& OutObje
 	const TraceServices::IChannelProvider* ChannelProvider = Session->ReadProvider<TraceServices::IChannelProvider>("ChannelProvider");
 	if (ChannelProvider)
 	{
-		const uint64 ChannelCount = ChannelProvider->GetChannelCount();		
 		const TArray<TraceServices::FChannelEntry>& Channels = ChannelProvider->GetChannels();
-		for (uint64 ChannelIndex = 0; ChannelIndex < Channels.Num(); ++ChannelIndex)
+		for (int32 ChannelIndex = 0; ChannelIndex < Channels.Num(); ++ChannelIndex)
 		{
 			FTraceObjectInfo& EventInfo = OutObjects.AddDefaulted_GetRef();
 			EventInfo.Name = Channels[ChannelIndex].Name;
@@ -83,9 +82,8 @@ void FBaseSessionFilterService::DisableAllChannels()
 	const TraceServices::IChannelProvider* ChannelProvider = Session->ReadProvider<TraceServices::IChannelProvider>("ChannelProvider");
 	if (ChannelProvider)
 	{
-		const uint64 ChannelCount = ChannelProvider->GetChannelCount();
 		const TArray<TraceServices::FChannelEntry>& Channels = ChannelProvider->GetChannels();
-		for (uint64 ChannelIndex = 0; ChannelIndex < Channels.Num(); ++ChannelIndex)
+		for (int32 ChannelIndex = 0; ChannelIndex < Channels.Num(); ++ChannelIndex)
 		{
 			SetObjectFilterState(Channels[ChannelIndex].Name, false);
 		}
diff --git a/Engine/Plugins/Developer/TraceDataFiltering/Source/TraceDataFiltering/Private/Trace/SessionTraceFilterService.cpp b/Engine/Plugins/Developer/TraceDataFiltering/Source/TraceDataFiltering/Private/Trace/SessionTraceFilterService.cpp
--- a/Engine/Plugins/Developer/TraceDataFiltering/Source/TraceDataFiltering/Private/Trace/SessionTraceFilterService.cpp
+++ b/Engine/Plugins/Developer/TraceDataFiltering/Source/TraceDataFiltering/Private/Trace/SessionTraceFilterService.cpp
@@ -37,7 +37,8 @@ void FSessionTraceFilterService::OnApplyChannelChanges()
 		return;
 	}
 
-	const UE::Trace::FStoreClient::FSessionInfo* SessionInfo = StoreClient->GetSessionInfoByTraceId(Handle);
+	const uint32 TraceId = static_cast<uint32>(Handle);
+	const UE::Trace::FStoreClient::FSessionInfo* SessionInfo = StoreClient->GetSessionInfoByTraceId(TraceId);
 
 	if (!SessionInfo)
 	{
