Deprecate TraceDataFilters plugin because the functionality has been moved to the Trace Control Widget from Session Frontend and the Live Trace Control Widget from Unreal Insights.

diff --git a/Engine/Plugins/Developer/TraceDataFiltering/TraceDataFilters.uplugin b/Engine/Plugins/Developer/TraceDataFiltering/TraceDataFilters.uplugin
--- a/Engine/Plugins/Developer/TraceDataFiltering/TraceDataFilters.uplugin
+++ b/Engine/Plugins/Developer/TraceDataFiltering/TraceDataFilters.uplugin
@@ -2,15 +2,15 @@
 	"FileVersion": 1,
 	"Version": 1,
 	"VersionName": "1.0",
-	"FriendlyName": "Trace Data Filtering",
-	"Description": "Allows for turning on/off individual or sets of Trace Channels.",
+	"FriendlyName": "Trace Data Filtering (Deprecated)",
+	"Description": "Allows for turning on/off individual or sets of Trace Channels. The plugin is deprecated because the functionality has been moved to the Trace Control Widget from Session Frontend and the Live Trace Control Widget from Unreal Insights.",
 	"Category": "Insights",
 	"CreatedBy": "Epic Games, Inc.",
 	"CreatedByURL": "https://epicgames.com",
 	"DocsURL": "",
 	"MarketplaceURL": "",
 	"SupportURL": "",
-	"EnabledByDefault": true,
+	"EnabledByDefault": false,
 	"CanContainContent": false,
 	"IsBetaVersion": false,
 	"Installed": false,
diff --git a/Engine/Programs/UnrealInsights/Config/DefaultEngine.ini b/Engine/Programs/UnrealInsights/Config/DefaultEngine.ini
--- a/Engine/Programs/UnrealInsights/Config/DefaultEngine.ini
+++ b/Engine/Programs/UnrealInsights/Config/DefaultEngine.ini
@@ -1,7 +1,6 @@
 [Plugins]
 +ProgramEnabledPlugins=AudioInsights
 +ProgramEnabledPlugins=GameplayInsights
-+ProgramEnabledPlugins=TraceDataFilters
 +ProgramEnabledPlugins=SlateInsights
 +ProgramEnabledPlugins=RenderGraphInsights
 +ProgramEnabledPlugins=IoStoreInsights
